Fix additional plugin modules that are causing the whole intermediate directory to be added to the installed build.

diff --git a/Engine/Plugins/Experimental/RemoteSession/Source/RemoteSession/RemoteSession.Build.cs b/Engine/Plugins/Experimental/RemoteSession/Source/RemoteSession/RemoteSession.Build.cs
--- a/Engine/Plugins/Experimental/RemoteSession/Source/RemoteSession/RemoteSession.Build.cs
+++ b/Engine/Plugins/Experimental/RemoteSession/Source/RemoteSession/RemoteSession.Build.cs
@@ -10,7 +10,7 @@ public class RemoteSession : ModuleRules
 		
 		PublicIncludePaths.AddRange(
 			new string[] {
-				".."
+				System.IO.Path.Combine(ModuleDirectory, "..")
 				// ... add public include paths required here ...
 			}
 		);
