Merging RemoteSession changes in 4062187 from //UE4/Dev-VR to Dev-Main to move work to another stream without creating conflicts (Part1)

diff --git a/Engine/Plugins/Experimental/BackChannel/BackChannel.uplugin b/Engine/Plugins/Experimental/BackChannel/BackChannel.uplugin
--- a/Engine/Plugins/Experimental/BackChannel/BackChannel.uplugin
+++ b/Engine/Plugins/Experimental/BackChannel/BackChannel.uplugin
@@ -1,22 +1,20 @@
 {
-	"FileVersion" : 3,
-	"FriendlyName" : "BackChannel",
-	"Version" : 1.0,
+	"FileVersion": 3,
+	"FriendlyName": "BackChannel",
+	"Version": 1.0,
 	"VersionName": "1",
-	"Description" : "BackChannel is an experimetal plugin that allows external tools and apps to query for and push data into a running Unreal session.",
-	"Category" : "Other",
-	"CreatedBy" : "Andrew Grant",
-	"CreatedByURL" : "",
+	"Description": "BackChannel is an experimetal plugin that allows external tools and apps to query for and push data into a running Unreal session.",
+	"Category": "Experimental",
+	"CreatedBy": "Andrew Grant",
+	"CreatedByURL": "",
 	"EnabledByDefault": true,
- 
-	"Modules" :
-	[
+
+	"Modules": [
 		{
-			"Name" : "BackChannel",
-			"Type" : "RuntimeNoCommandlet",
-			"LoadingPhase" : "Default",
-			"BlacklistPlatforms" :
-			[
+			"Name": "BackChannel",
+			"Type": "RuntimeNoCommandlet",
+			"LoadingPhase": "Default",
+			"BlacklistPlatforms": [
 				"HTML5"
 			]
 		}
diff --git a/Engine/Plugins/Experimental/BackChannel/README.md b/Engine/Plugins/Experimental/BackChannel/README.md
new file mode 100644
--- /dev/null
+++ b/Engine/Plugins/Experimental/BackChannel/README.md
@@ -0,0 +1,5 @@
+# BackChannel
+
+BackChannel is an experimetal plugin that allows external tools and apps to query for and push data into a running Unreal session.
+
+The goal is to support 2-3 ubiquitous protocols where implementations can be easily created or found for popular languages. The short term focus is OSC over TCP/IP and JSON REST over HTTP.
diff --git a/Engine/Plugins/Experimental/BackChannel/Source/BackChannel/Private/Transport/TestBackChannelTransport.cpp b/Engine/Plugins/Experimental/BackChannel/Source/BackChannel/Private/Transport/TestBackChannelTransport.cpp
--- a/Engine/Plugins/Experimental/BackChannel/Source/BackChannel/Private/Transport/TestBackChannelTransport.cpp
+++ b/Engine/Plugins/Experimental/BackChannel/Source/BackChannel/Private/Transport/TestBackChannelTransport.cpp
@@ -70,7 +70,7 @@ public:
 
 			FPlatformProcess::SleepNoStats(0.5);
 
-		} while (!ClientConnected);
+		} while (!ClientConnected || !AcceptConnected);
 
 		return AcceptedConnection.IsValid();
 	}
diff --git a/Engine/Plugins/Experimental/RemoteSession/README.md b/Engine/Plugins/Experimental/RemoteSession/README.md
new file mode 100644
--- /dev/null
+++ b/Engine/Plugins/Experimental/RemoteSession/README.md
@@ -0,0 +1,36 @@
+# RemoteSession
+A plugin for Unreal that allows one instance to act as a thin-client (rendering and input) to a second instance
+
+## Setup Info
+
+* Get the BackChannel plugin ([REDACTED_PROFILE_URL]/BackChannel)
+* Copy RemoteSession and BackChannel to your MyProject/PluginsFolder
+* Build
+* Start your project with -game or use Play In New Editor window (Play In Viewport is not currently supported)
+* Start the RemoteSessionApp on your mobile device and enter the IP address of the machine running the game
+
+
+## Useful Options
+
+There are a number of options that can be configured by putting one or more of the following in your DefaultEngine.ini under [RemoteSession]
+
+The values shown below are the current defaults. If you do not want to change from the defaults then I recommend not adding the line to your ini file incase they change in the future.
+
+<pre>
+[RemoteSession]
+; Start a host when the game runs
+bAutoHostWithGame=true
+; Start a host when PIE is used
+bAutoHostWithPIE=true
+; Image Quality (1-100)
+Quality=70
+; Framerate (0-100)
+Framerate=30
+; Port to listen on
+HostPort=2049
+; Whether RemoteSession runs in shipping builds
+bAllowInShipping=false
+</pre>
+
+
+Framerate and Quality can be adjusted at runtime via the remote.framerate and remote.quality cvars.
diff --git a/Engine/Plugins/Experimental/RemoteSession/RemoteSession.uplugin b/Engine/Plugins/Experimental/RemoteSession/RemoteSession.uplugin
--- a/Engine/Plugins/Experimental/RemoteSession/RemoteSession.uplugin
+++ b/Engine/Plugins/Experimental/RemoteSession/RemoteSession.uplugin
@@ -1,37 +1,41 @@
-{
-	"FileVersion": 3,
-	"Version": 1,
-	"VersionName": "1.0",
-	"FriendlyName": "RemoteSession",
-	"Description": "A plugin for Unreal that allows one instance to act as a thin-client (rendering and input) to a second instance",
-	"Category": "Other",
-	"CreatedBy": "Andrew Grant",
-	"CreatedByURL": "",
-	"DocsURL": "",
-	"MarketplaceURL": "",
-	"SupportURL": "",
-	"CanContainContent": true,
-	"IsBetaVersion": false,
-	"Installed": false,
-	"EnabledByDefault": true,
-
-	"Modules":
-	[
-		{
-			"Name": "RemoteSession",
-			"Type": "RuntimeNoCommandlet",
-			"LoadingPhase": "Default",
-			"BlacklistPlatforms" :
-			[
-				"HTML5"
-			]
-		}
-	],
-	"Plugins":
-	[
-        {
-            "Name": "BackChannel",
-            "Enabled": true
-        }
-    ]
+{
+	"FileVersion": 3,
+	"Version": 1,
+	"VersionName": "1.0",
+	"FriendlyName": "RemoteSession",
+	"Description": "A plugin for Unreal that allows one instance to act as a thin-client (rendering and input) to a second instance",
+	"Category": "Experimental",
+	"CreatedBy": "Andrew Grant",
+	"CreatedByURL": "",
+	"DocsURL": "",
+	"MarketplaceURL": "",
+	"SupportURL": "",
+	"CanContainContent": true,
+	"IsBetaVersion": false,
+	"Installed": false,
+	"EnabledByDefault": false,
+
+	"Modules":
+	[
+		{
+			"Name": "RemoteSession",
+			"Type": "Runtime",
+			"LoadingPhase": "PostConfigInit",
+			"BlacklistPlatforms" :
+			[
+				"HTML5"
+			]
+		}
+	],
+	"Plugins":
+	[
+		{
+			"Name": "BackChannel",
+			"Enabled": true
+		},
+		{
+			"Name": "AppleImageUtils",
+			"Enabled": true
+		}
+	]
 }
\ No newline at end of file
diff --git a/Engine/Plugins/Experimental/RemoteSession/Source/RemoteSession/Channels/RemoteSessionInputChannel.h b/Engine/Plugins/Experimental/RemoteSession/Source/RemoteSession/Channels/RemoteSessionInputChannel.h
--- a/Engine/Plugins/Experimental/RemoteSession/Source/RemoteSession/Channels/RemoteSessionInputChannel.h
+++ b/Engine/Plugins/Experimental/RemoteSession/Source/RemoteSession/Channels/RemoteSessionInputChannel.h
@@ -3,7 +3,7 @@
 #pragma once
 
 #include "RemoteSessionChannel.h"
-#include "MessageHandler/RecordingMessageHandler.h"
+#include "../Private/MessageHandler/RecordingMessageHandler.h"
 
 class FBackChannelOSCMessage;
 class FBackChannelOSCDispatch;
@@ -24,6 +24,8 @@ public:
 
 	void SetPlaybackWindow(TWeakPtr<SWindow> InWindow, TWeakPtr<FSceneViewport> InViewport);
 
+	void SetInputRect(const FVector2D& TopLeft, const FVector2D& Extents);
+
 	static FString StaticType();
 	virtual FString GetType() const override { return StaticType(); }
 
diff --git a/Engine/Plugins/Experimental/RemoteSession/Source/RemoteSession/Private/Channels/RemoteSessionFrameBufferChannel.cpp b/Engine/Plugins/Experimental/RemoteSession/Source/RemoteSession/Private/Channels/RemoteSessionFrameBufferChannel.cpp
--- a/Engine/Plugins/Experimental/RemoteSession/Source/RemoteSession/Private/Channels/RemoteSessionFrameBufferChannel.cpp
+++ b/Engine/Plugins/Experimental/RemoteSession/Source/RemoteSession/Private/Channels/RemoteSessionFrameBufferChannel.cpp
@@ -202,9 +202,9 @@ void FRemoteSessionFrameBufferChannel::Tick(const float InDeltaTime)
 			FUpdateTextureRegion2D* Region = new FUpdateTextureRegion2D(0, 0, 0, 0, QueuedImage->Width, QueuedImage->Height);
 			TArray<uint8>* TextureData = new TArray<uint8>(MoveTemp(QueuedImage->ImageData));
 
-			DecodedTextures[NextImage]->UpdateTextureRegions(0, 1, Region, 4 * QueuedImage->Width, 8, TextureData->GetData(), [this, NextImage](auto InTextureData, auto InRegions) {
+			DecodedTextures[NextImage]->UpdateTextureRegions(0, 1, Region, 4 * QueuedImage->Width, 8, TextureData->GetData(), [this, NextImage, TextureData](auto InTextureData, auto InRegions) {
 				DecodedTextureIndex = NextImage;
-				delete InTextureData;
+				delete TextureData; // delete array, not underlying data that UpdateTextureRegions passes us
 				delete InRegions;
 			});
 
diff --git a/Engine/Plugins/Experimental/RemoteSession/Source/RemoteSession/Private/Channels/RemoteSessionInputChannel.cpp b/Engine/Plugins/Experimental/RemoteSession/Source/RemoteSession/Private/Channels/RemoteSessionInputChannel.cpp
--- a/Engine/Plugins/Experimental/RemoteSession/Source/RemoteSession/Private/Channels/RemoteSessionInputChannel.cpp
+++ b/Engine/Plugins/Experimental/RemoteSession/Source/RemoteSession/Private/Channels/RemoteSessionInputChannel.cpp
@@ -71,6 +71,14 @@ void FRemoteSessionInputChannel::SetPlaybackWindow(TWeakPtr<SWindow> InWindow, T
 	PlaybackHandler->SetPlaybackWindow(InWindow, InViewport);
 }
 
+void FRemoteSessionInputChannel::SetInputRect(const FVector2D& TopLeft, const FVector2D& Extents)
+{
+	if (RecordingHandler.IsValid())
+	{
+		RecordingHandler->SetInputRect(TopLeft, Extents);
+	}
+}
+
 void FRemoteSessionInputChannel::Tick(const float InDeltaTime)
 {
 	// everything happens via messaging.
diff --git a/Engine/Plugins/Experimental/RemoteSession/Source/RemoteSession/Private/MessageHandler/Messages.h b/Engine/Plugins/Experimental/RemoteSession/Source/RemoteSession/Private/MessageHandler/Messages.h
--- a/Engine/Plugins/Experimental/RemoteSession/Source/RemoteSession/Private/MessageHandler/Messages.h
+++ b/Engine/Plugins/Experimental/RemoteSession/Source/RemoteSession/Private/MessageHandler/Messages.h
@@ -31,6 +31,8 @@ struct TwoParamMsg
 
 	TwoParamMsg(FArchive& Ar)
 	{
+		Param1 = P1();
+		Param2 = P2();
 		Ar << Param1;
 		Ar << Param2;
 	}
@@ -58,6 +60,9 @@ struct ThreeParamMsg
 
 	ThreeParamMsg(FArchive& Ar)
 	{
+		Param1 = P1();
+		Param2 = P2();
+		Param3 = P3();
 		Ar << Param1;
 		Ar << Param2;
 		Ar << Param3;
@@ -88,6 +93,10 @@ struct FourParamMsg
 
 	FourParamMsg(FArchive& Ar)
 	{
+		Param1 = P1();
+		Param2 = P2();
+		Param3 = P3();
+		Param4 = P4();
 		Ar << Param1;
 		Ar << Param2;
 		Ar << Param3;
@@ -121,6 +130,11 @@ struct FiveParamMsg
 
 	FiveParamMsg(FArchive& Ar)
 	{
+		Param1 = P1();
+		Param2 = P2();
+		Param3 = P3();
+		Param4 = P4();
+		Param5 = P5();
 		Ar << Param1;
 		Ar << Param2;
 		Ar << Param3;
diff --git a/Engine/Plugins/Experimental/RemoteSession/Source/RemoteSession/Private/MessageHandler/ProxyMessageHandler.cpp b/Engine/Plugins/Experimental/RemoteSession/Source/RemoteSession/Private/MessageHandler/ProxyMessageHandler.cpp
--- a/Engine/Plugins/Experimental/RemoteSession/Source/RemoteSession/Private/MessageHandler/ProxyMessageHandler.cpp
+++ b/Engine/Plugins/Experimental/RemoteSession/Source/RemoteSession/Private/MessageHandler/ProxyMessageHandler.cpp
@@ -198,21 +198,37 @@ void FProxyMessageHandler::OnEndGesture()
 	}
 }
 
+#if REMOTE_WITH_FORCE_PARAM
+bool FProxyMessageHandler::OnTouchStarted(const TSharedPtr< FGenericWindow >& Window, const FVector2D& Location, float Force, int32 TouchIndex, int32 ControllerId)
+#else
 bool FProxyMessageHandler::OnTouchStarted(const TSharedPtr< FGenericWindow >& Window, const FVector2D& Location, int32 TouchIndex, int32 ControllerId)
+#endif // REMOTE_WITH_FORCE_PARAM
 {
 	if (TargetHandler.IsValid())
 	{
+#if REMOTE_WITH_FORCE_PARAM
+		return TargetHandler->OnTouchStarted(Window, Location, Force, TouchIndex, ControllerId);
+#else
 		return TargetHandler->OnTouchStarted(Window, Location, TouchIndex, ControllerId);
+#endif
 	}
 
 	return false;
 }
 
+#if REMOTE_WITH_FORCE_PARAM
+bool FProxyMessageHandler::OnTouchMoved(const FVector2D& Location, float Force, int32 TouchIndex, int32 ControllerId)
+#else
 bool FProxyMessageHandler::OnTouchMoved(const FVector2D& Location, int32 TouchIndex, int32 ControllerId)
+#endif // REMOTE_WITH_FORCE_PARAM
 {
 	if (TargetHandler.IsValid())
 	{
+#if REMOTE_WITH_FORCE_PARAM
+		return TargetHandler->OnTouchMoved(Location, Force, TouchIndex, ControllerId);
+#else
 		return TargetHandler->OnTouchMoved(Location, TouchIndex, ControllerId);
+#endif
 	}
 
 	return false;
diff --git a/Engine/Plugins/Experimental/RemoteSession/Source/RemoteSession/Private/MessageHandler/ProxyMessageHandler.h b/Engine/Plugins/Experimental/RemoteSession/Source/RemoteSession/Private/MessageHandler/ProxyMessageHandler.h
--- a/Engine/Plugins/Experimental/RemoteSession/Source/RemoteSession/Private/MessageHandler/ProxyMessageHandler.h
+++ b/Engine/Plugins/Experimental/RemoteSession/Source/RemoteSession/Private/MessageHandler/ProxyMessageHandler.h
@@ -3,8 +3,14 @@
 #pragma once
 
 #include "CoreMinimal.h"
+#include "Launch/Resources/Version.h"
 #include "GenericPlatform/GenericApplicationMessageHandler.h"
 
+// For compatibility with a change that has not yet made it to //UE4/Main.
+// Please let andrew.grant know when you run into this as we may want to do some
+// cleanup if it occurs prior to the 4.20 branch.
+#define REMOTE_WITH_FORCE_PARAM 0
+
 class FProxyMessageHandler : public FGenericApplicationMessageHandler
 {
 	
@@ -56,9 +62,15 @@ public:
 
 	virtual void OnEndGesture() override;
 
+#if REMOTE_WITH_FORCE_PARAM
+	virtual bool OnTouchStarted(const TSharedPtr< FGenericWindow >& Window, const FVector2D& Location, float Force, int32 TouchIndex, int32 ControllerId) override;
+
+	virtual bool OnTouchMoved(const FVector2D& Location, float Force, int32 TouchIndex, int32 ControllerId) override;
+#else
 	virtual bool OnTouchStarted(const TSharedPtr< FGenericWindow >& Window, const FVector2D& Location, int32 TouchIndex, int32 ControllerId) override;
 
 	virtual bool OnTouchMoved(const FVector2D& Location, int32 TouchIndex, int32 ControllerId) override;
+#endif // REMOTE_WITH_FORCE_PARAM
 
 	virtual bool OnTouchEnded(const FVector2D& Location, int32 TouchIndex, int32 ControllerId) override;
 
diff --git a/Engine/Plugins/Experimental/RemoteSession/Source/RemoteSession/Private/MessageHandler/RecordingMessageHandler.cpp b/Engine/Plugins/Experimental/RemoteSession/Source/RemoteSession/Private/MessageHandler/RecordingMessageHandler.cpp
--- a/Engine/Plugins/Experimental/RemoteSession/Source/RemoteSession/Private/MessageHandler/RecordingMessageHandler.cpp
+++ b/Engine/Plugins/Experimental/RemoteSession/Source/RemoteSession/Private/MessageHandler/RecordingMessageHandler.cpp
@@ -32,6 +32,10 @@ FRecordingMessageHandler::FRecordingMessageHandler(const TSharedPtr<FGenericAppl
 {
 	OutputWriter = nullptr;
 	ConsumeInput = false;
+    bIsTouching = false;
+    InputRect = FRect(EForceInit::ForceInitToZero);
+    LastTouchLocation = FVector2D(EForceInit::ForceInitToZero);
+    
 
 	BIND_PLAYBACK_HANDLER(TEXT("OnKeyChar"), PlayOnKeyChar);
 	BIND_PLAYBACK_HANDLER(TEXT("OnKeyUp"), PlayOnKeyUp);
@@ -72,11 +76,30 @@ void FRecordingMessageHandler::SetPlaybackWindow(TWeakPtr<SWindow> InWindow, TWe
 	PlaybackViewport = InViewport;
 }
 
-FVector2D FRecordingMessageHandler::ConvertToNormalizedScreenLocation(const FVector2D& Location)
+void FRecordingMessageHandler::SetInputRect(const FVector2D& TopLeft, const FVector2D& Extents)
 {
-	const FVector2D ViewportSize = FVector2D(GEngine->GameViewport->Viewport->GetSizeXY());
+	InputRect = FRect(TopLeft, Extents);
+}
+
+bool FRecordingMessageHandler::ConvertToNormalizedScreenLocation(const FVector2D& InLocation, FVector2D& OutLocation)
+{
+	FRect ClipRect = InputRect;
+	FIntPoint Point = FIntPoint((int)InLocation.X, (int)InLocation.Y);
+
+	if (ClipRect.Width == 0 || ClipRect.Height == 0)
+	{
+		ClipRect = FRect(FVector2D(EForceInit::ForceInitToZero), FVector2D(GEngine->GameViewport->Viewport->GetSizeXY()));
+	}
 
-	return FVector2D(Location.X / ViewportSize.X, Location.Y / ViewportSize.Y);
+	if (!ClipRect.Contains(Point))
+	{
+		OutLocation = FVector2D(EForceInit::ForceInitToZero);
+		return false;
+	}
+
+	OutLocation = FVector2D((InLocation.X-ClipRect.X) / ClipRect.Width, (InLocation.Y - ClipRect.Y) / ClipRect.Height);
+
+	return true;
 }
 
 FVector2D FRecordingMessageHandler::ConvertFromNormalizedScreenLocation(const FVector2D& ScreenLocation)
@@ -185,12 +208,28 @@ void FRecordingMessageHandler::PlayOnKeyUp(FArchive& Ar)
 	OnKeyUp(Msg.Param1, Msg.Param2, Msg.Param3);
 }
 
+#if REMOTE_WITH_FORCE_PARAM
+bool FRecordingMessageHandler::OnTouchStarted(const TSharedPtr< FGenericWindow >& Window, const FVector2D& Location, float Force, int32 TouchIndex, int32 ControllerId)
+#else
 bool FRecordingMessageHandler::OnTouchStarted(const TSharedPtr< FGenericWindow >& Window, const FVector2D& Location, int32 TouchIndex, int32 ControllerId)
+#endif
 {
 	if (IsRecording())
 	{
-		ThreeParamMsg<FVector2D, int32, int32> Msg(ConvertToNormalizedScreenLocation(Location), TouchIndex, ControllerId);
-		RecordMessage(TEXT("OnTouchStarted"), Msg.AsData());
+		FVector2D Normalized;
+
+#if !REMOTE_WITH_FORCE_PARAM
+		float Force = 1.0f;
+#endif
+
+		if (ConvertToNormalizedScreenLocation(Location, Normalized))
+		{
+			// note - force is serialized last for backwards compat - force was introduced in 4.20
+			FourParamMsg<FVector2D, int32, int32, float> Msg(Normalized, TouchIndex, ControllerId, Force);
+			RecordMessage(TEXT("OnTouchStarted"), Msg.AsData());
+			bIsTouching = true;
+			LastTouchLocation = Location;
+		}
 	}
 
 	if (ConsumeInput)
@@ -198,12 +237,16 @@ bool FRecordingMessageHandler::OnTouchStarted(const TSharedPtr< FGenericWindow >
 		return true;
 	}
 
+#if REMOTE_WITH_FORCE_PARAM
+	return FProxyMessageHandler::OnTouchStarted(Window, Location, Force, TouchIndex, ControllerId);
+#else
 	return FProxyMessageHandler::OnTouchStarted(Window, Location, TouchIndex, ControllerId);
+#endif
 }
 
 void FRecordingMessageHandler::PlayOnTouchStarted(FArchive& Ar)
 {
-	ThreeParamMsg<FVector2D, int32, int32 > Msg(Ar);
+	FourParamMsg<FVector2D, int32, int32, float > Msg(Ar);
 	FVector2D ScreenLocation = ConvertFromNormalizedScreenLocation(Msg.Param1);
 
 	TSharedPtr<FGenericWindow> Window;
@@ -213,15 +256,35 @@ void FRecordingMessageHandler::PlayOnTouchStarted(FArchive& Ar)
 		Window = PlaybackWindow.Pin()->GetNativeWindow();
 	}
 
+#if REMOTE_WITH_FORCE_PARAM
+	// note - force is serialized last for backwards compat - force was introduced in 4.20
+	OnTouchStarted(Window, ScreenLocation, Msg.Param4, Msg.Param2, Msg.Param3 );
+#else
 	OnTouchStarted(Window, ScreenLocation, Msg.Param2, Msg.Param3);
+#endif
 }
 
+#if REMOTE_WITH_FORCE_PARAM
+bool FRecordingMessageHandler::OnTouchMoved(const FVector2D& Location, float Force, int32 TouchIndex, int32 ControllerId)
+#else
 bool FRecordingMessageHandler::OnTouchMoved(const FVector2D& Location, int32 TouchIndex, int32 ControllerId)
+#endif
 {
 	if (IsRecording())
 	{
-		ThreeParamMsg<FVector2D, int32, int32> Msg(ConvertToNormalizedScreenLocation(Location), TouchIndex, ControllerId);
-		OutputWriter->RecordMessage(TEXT("OnTouchMoved"), Msg.AsData());
+		FVector2D Normalized;
+
+#if !REMOTE_WITH_FORCE_PARAM
+		float Force = 1.0f;
+#endif
+		if (ConvertToNormalizedScreenLocation(Location, Normalized))
+		{
+			// note - force is serialized last for backwards compat - force was introduced in 4.20
+			FourParamMsg<FVector2D, int32, int32, float> Msg(Normalized, TouchIndex, ControllerId, Force);
+			OutputWriter->RecordMessage(TEXT("OnTouchMoved"), Msg.AsData());
+			bIsTouching = true;
+			LastTouchLocation = Location;
+		}
 	}
 
 	if (ConsumeInput)
@@ -229,22 +292,40 @@ bool FRecordingMessageHandler::OnTouchMoved(const FVector2D& Location, int32 Tou
 		return true;
 	}
 
+#if REMOTE_WITH_FORCE_PARAM
+	return FProxyMessageHandler::OnTouchMoved(Location, Force, TouchIndex, ControllerId);
+#else
 	return FProxyMessageHandler::OnTouchMoved(Location, TouchIndex, ControllerId);
+#endif
 }
 
 void FRecordingMessageHandler::PlayOnTouchMoved(FArchive& Ar)
 {
-	ThreeParamMsg<FVector2D, int32, int32 > Msg(Ar);
+	FourParamMsg<FVector2D, int32, int32, float > Msg(Ar);
 	FVector2D ScreenLocation = ConvertFromNormalizedScreenLocation(Msg.Param1);
+#if REMOTE_WITH_FORCE_PARAM
+	// note - force is serialized last for backwards compat - force was introduced in 4.20
+	OnTouchMoved(ScreenLocation, Msg.Param4, Msg.Param2, Msg.Param3);
+#else
 	OnTouchMoved(ScreenLocation, Msg.Param2, Msg.Param3);
+#endif
 }
 
 bool FRecordingMessageHandler::OnTouchEnded(const FVector2D& Location, int32 TouchIndex, int32 ControllerId)
 {
 	if (IsRecording())
 	{
-		ThreeParamMsg<FVector2D, int32, int32> Msg(ConvertToNormalizedScreenLocation(Location), TouchIndex, ControllerId);
-		OutputWriter->RecordMessage(TEXT("OnTouchEnded"), Msg.AsData());
+		FVector2D Normalized;
+
+        // if outside our bounds, end the touch where it left
+		if (ConvertToNormalizedScreenLocation(Location, Normalized) == false)
+        {
+            ConvertToNormalizedScreenLocation(LastTouchLocation, Normalized);
+        }
+		
+        ThreeParamMsg<FVector2D, int32, int32> Msg(Normalized, TouchIndex, ControllerId);
+        OutputWriter->RecordMessage(TEXT("OnTouchEnded"), Msg.AsData());
+        bIsTouching = false;
 	}
 
 	if (ConsumeInput)
diff --git a/Engine/Plugins/Experimental/RemoteSession/Source/RemoteSession/Private/MessageHandler/RecordingMessageHandler.h b/Engine/Plugins/Experimental/RemoteSession/Source/RemoteSession/Private/MessageHandler/RecordingMessageHandler.h
--- a/Engine/Plugins/Experimental/RemoteSession/Source/RemoteSession/Private/MessageHandler/RecordingMessageHandler.h
+++ b/Engine/Plugins/Experimental/RemoteSession/Source/RemoteSession/Private/MessageHandler/RecordingMessageHandler.h
@@ -19,6 +19,41 @@ DECLARE_DELEGATE_OneParam(FRecordedMessageDispatch, FArchive&);
 
 class FRecordingMessageHandler : public FProxyMessageHandler, public TSharedFromThis<FRecordingMessageHandler>
 {
+	struct FRect
+	{
+		FRect() {}
+
+		FRect(EForceInit)
+			: X(0.0f)
+			, Y(0.0f)
+			, Width(0.0f)
+			, Height(0.0f)
+		{}
+
+		FRect(float InX, float InY, float InWidth, float InHeight)
+			: X(InX)
+			, Y(InY)
+			, Width(InWidth)
+			, Height(InHeight)
+		{}
+
+		FRect(const FVector2D& Point, const FVector2D& Extents)
+			: X(Point.X)
+			, Y(Point.Y)
+			, Width(Extents.X)
+			, Height(Extents.Y)
+		{}
+
+		FORCEINLINE bool Contains(FVector2D P) const
+		{
+			return P.X >= X && P.X < (X+Width) && P.Y >= Y && P.Y < (Y+Height);
+		}
+
+		float	X;
+		float	Y;
+		float	Width;
+		float	Height;
+	};
 public:
 
 	FRecordingMessageHandler(const TSharedPtr<FGenericApplicationMessageHandler>& InTargetHandler);
@@ -34,6 +69,8 @@ public:
 
 	void SetPlaybackWindow(TWeakPtr<SWindow> InWindow, TWeakPtr<FSceneViewport> InViewport);
 
+	void SetInputRect(const FVector2D& TopLeft, const FVector2D& Extents);
+
 public:
 
 	virtual bool OnKeyChar(const TCHAR Character, const bool IsRepeat) override;
@@ -44,8 +81,13 @@ public:
 	virtual bool OnTouchGesture(EGestureEvent GestureType, const FVector2D& Delta, float WheelDelta, bool bIsDirectionInvertedFromDevice) override;
 	virtual void OnEndGesture() override;
 
-	virtual bool OnTouchStarted(const TSharedPtr< FGenericWindow >& Window, const FVector2D& Location, int32 TouchIndex, int32 ControllerId) override;
+#if REMOTE_WITH_FORCE_PARAM
+	virtual bool OnTouchStarted(const TSharedPtr< FGenericWindow >& Window, const FVector2D& Location, float Force, int32 TouchIndex, int32 ControllerId) override;
+	virtual bool OnTouchMoved(const FVector2D& Location, float Force, int32 TouchIndex, int32 ControllerId) override;
+#else
+	virtual bool OnTouchStarted(const TSharedPtr< FGenericWindow >& Window, const FVector2D& Location,  int32 TouchIndex, int32 ControllerId) override;
 	virtual bool OnTouchMoved(const FVector2D& Location, int32 TouchIndex, int32 ControllerId) override;
+#endif // #if REMOTE_WITH_FORCE_PARAM
 	virtual bool OnTouchEnded(const FVector2D& Location, int32 TouchIndex, int32 ControllerId) override;
 	virtual bool OnMotionDetected(const FVector& Tilt, const FVector& RotationRate, const FVector& Gravity, const FVector& Acceleration, int32 ControllerId) override;
 
@@ -53,7 +95,7 @@ public:
 
 protected:
 
-	FVector2D ConvertToNormalizedScreenLocation(const FVector2D& Location);
+	bool ConvertToNormalizedScreenLocation(const FVector2D& InLocation, FVector2D& OutLocation);
 	FVector2D ConvertFromNormalizedScreenLocation(const FVector2D& ScreenLocation);
 
 	void RecordMessage(const TCHAR* MsgName, const TArray<uint8>& Data);
@@ -77,4 +119,9 @@ protected:
 	TWeakPtr<FSceneViewport>			PlaybackViewport;
 
 	TMap<FString, FRecordedMessageDispatch> DispatchTable;
+
+	FRect								InputRect;
+    FVector2D                           LastTouchLocation;
+    bool                                bIsTouching;
+
 };
\ No newline at end of file
diff --git a/Engine/Plugins/Experimental/RemoteSession/Source/RemoteSession/Private/RemoteSessionHost.cpp b/Engine/Plugins/Experimental/RemoteSession/Source/RemoteSession/Private/RemoteSessionHost.cpp
--- a/Engine/Plugins/Experimental/RemoteSession/Source/RemoteSession/Private/RemoteSessionHost.cpp
+++ b/Engine/Plugins/Experimental/RemoteSession/Source/RemoteSession/Private/RemoteSessionHost.cpp
@@ -9,20 +9,24 @@
 #include "Channels/RemoteSessionXRTrackingChannel.h"
 #include "Channels/RemoteSessionFrameBufferChannel.h"
 #include "Engine/GameEngine.h"
+#include "Framework/Application/SlateApplication.h"
 
 #if WITH_EDITOR
 	#include "Editor.h"
 	#include "Editor/EditorEngine.h"
 #endif
 
-PRAGMA_DISABLE_OPTIMIZATION
-
+namespace RemoteSessionEd
+{
+	static FAutoConsoleVariable SlateDragDistanceOverride(TEXT("RemoteSessionEd.SlateDragDistanceOverride"), 10.0f, TEXT("How many pixels you need to drag before a drag and drop operation starts in remote app"));
+};
 
 
 FRemoteSessionHost::FRemoteSessionHost(int32 InQuality, int32 InFramerate)
 {
 	Quality = InQuality;
 	Framerate = InFramerate;
+	SavedEditorDragTriggerDistance = FSlateApplication::Get().GetDragTriggerDistance();
 }
 
 FRemoteSessionHost::~FRemoteSessionHost()
@@ -37,6 +41,17 @@ FRemoteSessionHost::~FRemoteSessionHost()
 	Close();
 }
 
+
+void FRemoteSessionHost::Close()
+{
+	FRemoteSessionRole::Close();
+
+	if (FSlateApplication::IsInitialized())
+	{
+		FSlateApplication::Get().SetDragTriggerDistance(SavedEditorDragTriggerDistance);
+	}
+}
+
 void FRemoteSessionHost::SetScreenSharing(const bool bEnabled)
 {
 }
@@ -60,7 +75,10 @@ bool FRemoteSessionHost::StartListening(const uint16 InPort)
 	{
 		Listener = Transport->CreateConnection(IBackChannelTransport::TCP);
 
-		Listener->Listen(InPort);
+		if (Listener->Listen(InPort) == false)
+		{
+			Listener = nullptr;
+		}
 	}
 
 	return Listener.IsValid();
@@ -95,7 +113,9 @@ bool FRemoteSessionHost::ProcessIncomingConnection(TSharedRef<IBackChannelConnec
 				}
 			}
 		}
-		
+
+		SavedEditorDragTriggerDistance = FSlateApplication::Get().GetDragTriggerDistance();
+		FSlateApplication::Get().SetDragTriggerDistance(RemoteSessionEd::SlateDragDistanceOverride->GetFloat());
 	}
 	else
 #endif
@@ -129,14 +149,10 @@ void FRemoteSessionHost::Tick(float DeltaTime)
 	// non-threaded listener
 	if (IsConnected() == false)
 	{
-		Listener->WaitForConnection(0, [this](TSharedRef<IBackChannelConnection> NewConnection) {
-			return ProcessIncomingConnection(NewConnection);
+		Listener->WaitForConnection(0, [this](TSharedRef<IBackChannelConnection> InConnection) {
+			return ProcessIncomingConnection(InConnection);
 		});
 	}
 	
 	FRemoteSessionRole::Tick(DeltaTime);
 }
-
-
-
-PRAGMA_ENABLE_OPTIMIZATION
\ No newline at end of file
diff --git a/Engine/Plugins/Experimental/RemoteSession/Source/RemoteSession/Private/RemoteSessionHost.h b/Engine/Plugins/Experimental/RemoteSession/Source/RemoteSession/Private/RemoteSessionHost.h
--- a/Engine/Plugins/Experimental/RemoteSession/Source/RemoteSession/Private/RemoteSessionHost.h
+++ b/Engine/Plugins/Experimental/RemoteSession/Source/RemoteSession/Private/RemoteSessionHost.h
@@ -17,6 +17,8 @@ public:
 	FRemoteSessionHost(int32 InQuality, int32 InFramerate);
 	~FRemoteSessionHost();
 
+	virtual void Close() override;
+
 	bool StartListening(const uint16 Port);
 
 	void SetScreenSharing(const bool bEnabled);
@@ -33,4 +35,7 @@ protected:
 
 	int32		Quality;
 	int32		Framerate;
+
+	/** Saved information about the editor and viewport we possessed, so we can restore it after exiting VR mode */
+	float SavedEditorDragTriggerDistance;
 };
\ No newline at end of file
diff --git a/Engine/Plugins/Experimental/RemoteSession/Source/RemoteSession/Private/RemoteSessionModule.cpp b/Engine/Plugins/Experimental/RemoteSession/Source/RemoteSession/Private/RemoteSessionModule.cpp
--- a/Engine/Plugins/Experimental/RemoteSession/Source/RemoteSession/Private/RemoteSessionModule.cpp
+++ b/Engine/Plugins/Experimental/RemoteSession/Source/RemoteSession/Private/RemoteSessionModule.cpp
@@ -27,8 +27,10 @@ protected:
 	int32								Framerate;
 
 	bool bAutoHostWithPIE;
+	bool bAutoHostWithGame;
 	FDelegateHandle PostPieDelegate;
 	FDelegateHandle EndPieDelegate;
+	FDelegateHandle GameStartDelegate;
 
 public:
 
@@ -39,11 +41,12 @@ public:
 
 	void StartupModule()
 	{
-		bool bAutoHostWithGame = false;
+		// set defaults
 		DefaultPort = IRemoteSessionModule::kDefaultPort;
 		Quality = 85;
 		Framerate = 30;
-		bAutoHostWithPIE = false;
+		bAutoHostWithPIE = true;
+		bAutoHostWithGame = true;
 
 		GConfig->GetBool(TEXT("RemoteSession"), TEXT("bAutoHostWithGame"), bAutoHostWithGame, GEngineIni);
 		GConfig->GetBool(TEXT("RemoteSession"), TEXT("bAutoHostWithPIE"), bAutoHostWithPIE, GEngineIni);
@@ -51,19 +54,16 @@ public:
 		GConfig->GetInt(TEXT("RemoteSession"), TEXT("Quality"), Quality, GEngineIni);
 		GConfig->GetInt(TEXT("RemoteSession"), TEXT("Framerate"), Framerate, GEngineIni);
 
-		if (PLATFORM_DESKTOP)
+
+		if (PLATFORM_DESKTOP 
+			&& IsRunningDedicatedServer() == false 
+			&& IsRunningCommandlet() == false)
 		{
-			if (GIsEditor)
-			{
 #if WITH_EDITOR
-				PostPieDelegate = FEditorDelegates::PostPIEStarted.AddRaw(this, &FRemoteSessionModule::OnPIEStarted);
-				EndPieDelegate = FEditorDelegates::EndPIE.AddRaw(this, &FRemoteSessionModule::OnPIEEnded);
+			PostPieDelegate = FEditorDelegates::PostPIEStarted.AddRaw(this, &FRemoteSessionModule::OnPIEStarted);
+			EndPieDelegate = FEditorDelegates::EndPIE.AddRaw(this, &FRemoteSessionModule::OnPIEEnded);
 #endif
-			}
-			else if (bAutoHostWithGame)
-			{
-				InitHost();
-			}
+			GameStartDelegate = FCoreDelegates::OnFEngineLoopInitComplete.AddRaw(this, &FRemoteSessionModule::OnGameStarted);
 		}
 	}
 
@@ -82,6 +82,38 @@ public:
 			FEditorDelegates::EndPIE.Remove(EndPieDelegate);
 		}
 #endif
+
+		if (GameStartDelegate.IsValid())
+		{
+			FCoreDelegates::OnFEngineLoopInitComplete.Remove(GameStartDelegate);
+		}
+	}
+
+	void OnGameStarted()
+	{
+		bool IsHostGame = PLATFORM_DESKTOP
+			&& GIsEditor == false
+			&& IsRunningDedicatedServer() == false
+			&& IsRunningCommandlet() == false;
+
+		if (IsHostGame && bAutoHostWithGame)
+		{
+			InitHost();
+		}
+	}
+
+	void OnPIEStarted(bool bSimulating)
+	{
+		if (bAutoHostWithPIE)
+		{
+			InitHost();
+		}
+	}
+
+	void OnPIEEnded(bool bSimulating)
+	{
+		// always stop, incase it was started via the console
+		StopHost();
 	}
 
 	virtual void InitClient(const TCHAR* RemoteAddress) override
@@ -112,12 +144,22 @@ public:
 
 	virtual void InitHost(const int16 Port = 0) override
 	{
-#if !UE_BUILD_SHIPPING
 		if (Host.IsValid())
 		{
 			Host = nullptr;
 		}
 
+#if UE_BUILD_SHIPPING
+		bool bAllowInShipping = false;
+		GConfig->GetBool(TEXT("RemoteSession"), TEXT("bAllowInShipping"), bAllowInShipping, GEngineIni);
+
+		if (bAllowInShipping == false)
+		{
+			UE_LOG(LogRemoteSession, Log, TEXT("RemoteSession is disabled. Shipping=1"));
+			return;
+		}
+#endif
+
 		TSharedPtr<FRemoteSessionHost> NewHost = MakeShareable(new FRemoteSessionHost(Quality, Framerate));
 
 		int16 SelectedPort = Port ? Port : (int16)DefaultPort;
@@ -131,9 +173,6 @@ public:
 		{
 			UE_LOG(LogRemoteSession, Error, TEXT("Failed to start host listening on port %d"), SelectedPort);
 		}
-#else
-		UE_LOG(LogRemoteSession, Log, TEXT("RemoteSession is disabled. Shipping=1"));
-#endif
 	}
 
 	virtual bool IsHostRunning() const override
@@ -156,20 +195,6 @@ public:
 		return Host;
 	}
 
-	void OnPIEStarted(bool bSimulating)
-	{
-		if (bAutoHostWithPIE)
-		{
-			InitHost();
-		}
-	}
-
-	void OnPIEEnded(bool bSimulating)
-	{
-		// always stop, incase it was started via the console
-		StopHost();
-	}
-
 	virtual TStatId GetStatId() const override
 	{
 		RETURN_QUICK_DECLARE_CYCLE_STAT(FRemoteSession, STATGROUP_Tickables);
diff --git a/Engine/Plugins/Experimental/RemoteSession/Source/RemoteSession/RemoteSessionRole.h b/Engine/Plugins/Experimental/RemoteSession/Source/RemoteSession/RemoteSessionRole.h
--- a/Engine/Plugins/Experimental/RemoteSession/Source/RemoteSession/RemoteSessionRole.h
+++ b/Engine/Plugins/Experimental/RemoteSession/Source/RemoteSession/RemoteSessionRole.h
@@ -11,6 +11,8 @@ class REMOTESESSION_API IRemoteSessionRole
 public:
 	virtual~IRemoteSessionRole() {}
 
+	virtual bool IsConnected() const = 0;
+
 	virtual TSharedPtr<IRemoteSessionChannel> GetChannel(const FString& Type) = 0;
 
 	template<class T>
