Copying //UE4/Dev-Documentation to //UE4/Samples-Main (Source: //UE4/Dev-Documentation @ 4091886)

diff --git a/Engine/Plugins/Experimental/BackChannel/Source/BackChannel/Private/Protocol/BackChannelOSCConnection.cpp b/Engine/Plugins/Experimental/BackChannel/Source/BackChannel/Private/Protocol/BackChannelOSCConnection.cpp
--- a/Engine/Plugins/Experimental/BackChannel/Source/BackChannel/Private/Protocol/BackChannelOSCConnection.cpp
+++ b/Engine/Plugins/Experimental/BackChannel/Source/BackChannel/Private/Protocol/BackChannelOSCConnection.cpp
@@ -287,7 +287,18 @@ void FBackChannelOSCConnection::SetMessageOptions(const TCHAR* Path, int32 MaxQu
 	MessageLimits.FindOrAdd(Path) = MaxQueuedMessages;
 }
 
+FDelegateHandle FBackChannelOSCConnection::AddMessageHandler(const TCHAR* Path, FBackChannelDispatchDelegate::FDelegate Delegate)
+{
+	FScopeLock Lock(&ReceiveMutex);
+	return DispatchMap.GetAddressHandler(Path).Add(Delegate);
+}
 
+/* Remove a delegate handle */
+void FBackChannelOSCConnection::RemoveMessageHandler(const TCHAR* Path, FDelegateHandle& InHandle)
+{
+	DispatchMap.GetAddressHandler(Path).Remove(InHandle);
+	InHandle.Reset();
+}
 
 int32 FBackChannelOSCConnection::GetMessageCountForPath(const TCHAR* Path)
 {
@@ -376,9 +387,3 @@ void FBackChannelOSCConnection::RemoveMessagesWithPath(const TCHAR* Path, const
 	}
 }
 
-
-FBackChannelOSCDispatch& FBackChannelOSCConnection::GetDispatchMap()
-{
-	return DispatchMap;
-}
-
diff --git a/Engine/Plugins/Experimental/BackChannel/Source/BackChannel/Private/Transport/BackChannelConnection.cpp b/Engine/Plugins/Experimental/BackChannel/Source/BackChannel/Private/Transport/BackChannelConnection.cpp
--- a/Engine/Plugins/Experimental/BackChannel/Source/BackChannel/Private/Transport/BackChannelConnection.cpp
+++ b/Engine/Plugins/Experimental/BackChannel/Source/BackChannel/Private/Transport/BackChannelConnection.cpp
@@ -279,7 +279,7 @@ bool FBackChannelConnection::WaitForConnection(double InTimeout, TFunction<bool(
 		else
 		{
 			CheckSucceeded = true;
-			HasConnection = State == ESocketConnectionState::SCS_Connected;
+			HasConnection = Socket->Wait(ESocketWaitConditions::WaitForWrite, SleepTime);
 		}
 	}
 	
diff --git a/Engine/Plugins/Experimental/BackChannel/Source/BackChannel/Protocol/OSC/BackChannelOSCConnection.h b/Engine/Plugins/Experimental/BackChannel/Source/BackChannel/Protocol/OSC/BackChannelOSCConnection.h
--- a/Engine/Plugins/Experimental/BackChannel/Source/BackChannel/Protocol/OSC/BackChannelOSCConnection.h
+++ b/Engine/Plugins/Experimental/BackChannel/Source/BackChannel/Protocol/OSC/BackChannelOSCConnection.h
@@ -43,15 +43,17 @@ public:
 	/* Send the provided OSC packet */
 	bool SendPacket(FBackChannelOSCPacket& Packet);
 
-	/* Return our DispatchMap. This can be modified as necessary */
-	FBackChannelOSCDispatch& GetDispatchMap();
+	/* Bind a delegate to a message address */
+	FDelegateHandle AddMessageHandler(const TCHAR* Path, FBackChannelDispatchDelegate::FDelegate Delegate);
+	
+	/* Remove a delegate handle */
+	void RemoveMessageHandler(const TCHAR* Path, FDelegateHandle& InHandle);
 
 	/* Set options for the specified message path */
 	void SetMessageOptions(const TCHAR* Path, int32 MaxQueuedMessages);
 
 	FString GetDescription();
 
-
 protected:
 	// Begin protected FRunnable overrides
 	virtual uint32 Run() override;
@@ -96,4 +98,4 @@ protected:
 	int32				ExpectedDataSize;
 	TArray<uint8>		ReceiveBuffer;
 
-};
\ No newline at end of file
+};
diff --git a/Engine/Plugins/Experimental/RemoteSession/RemoteSession.uplugin b/Engine/Plugins/Experimental/RemoteSession/RemoteSession.uplugin
--- a/Engine/Plugins/Experimental/RemoteSession/RemoteSession.uplugin
+++ b/Engine/Plugins/Experimental/RemoteSession/RemoteSession.uplugin
@@ -1,12 +1,12 @@
-{
+﻿{
 	"FileVersion": 3,
 	"Version": 1,
 	"VersionName": "1.0",
 	"FriendlyName": "RemoteSession",
 	"Description": "A plugin for Unreal that allows one instance to act as a thin-client (rendering and input) to a second instance",
-	"Category": "Other",
-	"CreatedBy": "Andrew Grant",
-	"CreatedByURL": "",
+	"Category": "Experimental",
+	"CreatedBy": "Epic Games, Inc.",
+	"CreatedByURL": "http://epicgames.com",
 	"DocsURL": "",
 	"MarketplaceURL": "",
 	"SupportURL": "",
@@ -15,27 +15,24 @@
 	"Installed": false,
 	"EnabledByDefault": false,
 
-	"Modules":
-	[
+	"Modules": [
 		{
 			"Name": "RemoteSession",
-			"Type": "RuntimeNoCommandlet",
+			"Type": "Runtime",
 			"LoadingPhase": "PostConfigInit",
-			"BlacklistPlatforms" :
-			[
+			"BlacklistPlatforms": [
 				"HTML5"
 			]
 		}
 	],
-	"Plugins":
-	[
-        {
-            "Name": "BackChannel",
-            "Enabled": true
-        },
-        {
-            "Name": "AppleImageUtils",
-            "Enabled": true
-        }
-    ]
+	"Plugins": [
+		{
+			"Name": "BackChannel",
+			"Enabled": true
+		},
+		{
+			"Name": "AppleImageUtils",
+			"Enabled": true
+		}
+	]
 }
\ No newline at end of file
diff --git a/Engine/Plugins/Experimental/RemoteSession/Source/RemoteSession/Channels/RemoteSessionARCameraChannel.h b/Engine/Plugins/Experimental/RemoteSession/Source/RemoteSession/Channels/RemoteSessionARCameraChannel.h
--- a/Engine/Plugins/Experimental/RemoteSession/Source/RemoteSession/Channels/RemoteSessionARCameraChannel.h
+++ b/Engine/Plugins/Experimental/RemoteSession/Source/RemoteSession/Channels/RemoteSessionARCameraChannel.h
@@ -26,8 +26,8 @@ public:
 	virtual void Tick(const float InDeltaTime) override;
 
 	/* Begin IRemoteSessionChannel implementation */
-	static FString StaticType() { return TEXT("rs.arcamera"); }
-	virtual FString GetType() const override { return StaticType(); }
+	static const TCHAR* StaticType() { return TEXT("FRemoteSessionARCameraChannel"); }
+	virtual const TCHAR* GetType() const override { return StaticType(); }
 	/* End IRemoteSessionChannel implementation */
 
 	/** Returns the post process material to render with. The textures are set on the fly */
diff --git a/Engine/Plugins/Experimental/RemoteSession/Source/RemoteSession/Channels/RemoteSessionChannel.h b/Engine/Plugins/Experimental/RemoteSession/Source/RemoteSession/Channels/RemoteSessionChannel.h
--- a/Engine/Plugins/Experimental/RemoteSession/Source/RemoteSession/Channels/RemoteSessionChannel.h
+++ b/Engine/Plugins/Experimental/RemoteSession/Source/RemoteSession/Channels/RemoteSessionChannel.h
@@ -23,6 +23,6 @@ public:
 
 	virtual void Tick(const float InDeltaTime) = 0;
 
-	virtual FString GetType() const = 0;
+	virtual const TCHAR* GetType() const = 0;
 
 };
diff --git a/Engine/Plugins/Experimental/RemoteSession/Source/RemoteSession/Channels/RemoteSessionFrameBufferChannel.h b/Engine/Plugins/Experimental/RemoteSession/Source/RemoteSession/Channels/RemoteSessionFrameBufferChannel.h
--- a/Engine/Plugins/Experimental/RemoteSession/Source/RemoteSession/Channels/RemoteSessionFrameBufferChannel.h
+++ b/Engine/Plugins/Experimental/RemoteSession/Source/RemoteSession/Channels/RemoteSessionFrameBufferChannel.h
@@ -36,10 +36,16 @@ public:
 	UTexture2D* GetHostScreen() const;
 
 	/* Begin IRemoteSessionChannel implementation */
-	static FString StaticType();
-	virtual FString GetType() const override { return StaticType(); }
+	static const TCHAR* StaticType() { return TEXT("FRemoteSessionFrameBufferChannel"); }
+	virtual const TCHAR* GetType() const override { return StaticType(); }
 	/* End IRemoteSessionChannel implementation */
 
+	/** Signals that the viewport was resized */
+	void OnViewportResized(FVector2D NewSize);
+
+	/** Safely create the frame grabber */
+	void CreateFrameGrabber(TSharedRef<FSceneViewport> Viewport);
+
 protected:
 
 	/** Underlying connection */
@@ -58,6 +64,9 @@ protected:
 	/** Creates a texture to receive images into */
 	void CreateTexture(const int32 InSlot, const int32 InWidth, const int32 InHeight);
 
+	/** Release the FrameGrabber*/
+	void ReleaseFrameGrabber();
+
 	TSharedPtr<FFrameGrabber>				FrameGrabber;
 	
 	struct FImageData
@@ -91,4 +100,10 @@ protected:
 
 	/** So we can manage callback lifetimes properly */
 	FDelegateHandle MessageCallbackHandle;
-};
\ No newline at end of file
+
+	/** Shows that the viewport was just resized */
+	bool ViewportResized;
+
+	/** Holds a reference to the scene viewport */
+	TSharedPtr<FSceneViewport> SceneViewport;
+};
diff --git a/Engine/Plugins/Experimental/RemoteSession/Source/RemoteSession/Channels/RemoteSessionInputChannel.h b/Engine/Plugins/Experimental/RemoteSession/Source/RemoteSession/Channels/RemoteSessionInputChannel.h
--- a/Engine/Plugins/Experimental/RemoteSession/Source/RemoteSession/Channels/RemoteSessionInputChannel.h
+++ b/Engine/Plugins/Experimental/RemoteSession/Source/RemoteSession/Channels/RemoteSessionInputChannel.h
@@ -3,7 +3,7 @@
 #pragma once
 
 #include "RemoteSessionChannel.h"
-#include "MessageHandler/RecordingMessageHandler.h"
+#include "../Private/MessageHandler/RecordingMessageHandler.h"
 
 class FBackChannelOSCMessage;
 class FBackChannelOSCDispatch;
@@ -24,8 +24,10 @@ public:
 
 	void SetPlaybackWindow(TWeakPtr<SWindow> InWindow, TWeakPtr<FSceneViewport> InViewport);
 
-	static FString StaticType();
-	virtual FString GetType() const override { return StaticType(); }
+	void SetInputRect(const FVector2D& TopLeft, const FVector2D& Extents);
+
+	static const TCHAR* StaticType() { return TEXT("FRemoteSessionInputChannel"); }
+	virtual const TCHAR* GetType() const override { return StaticType(); }
 
 protected:
 
@@ -41,4 +43,4 @@ protected:
 
 	/** So we can manage callback lifetimes properly */
 	FDelegateHandle MessageCallbackHandle;
-};
\ No newline at end of file
+};
diff --git a/Engine/Plugins/Experimental/RemoteSession/Source/RemoteSession/Channels/RemoteSessionXRTrackingChannel.h b/Engine/Plugins/Experimental/RemoteSession/Source/RemoteSession/Channels/RemoteSessionXRTrackingChannel.h
--- a/Engine/Plugins/Experimental/RemoteSession/Source/RemoteSession/Channels/RemoteSessionXRTrackingChannel.h
+++ b/Engine/Plugins/Experimental/RemoteSession/Source/RemoteSession/Channels/RemoteSessionXRTrackingChannel.h
@@ -41,8 +41,8 @@ public:
 	void ReceiveXRTracking(FBackChannelOSCMessage& Message, FBackChannelOSCDispatch& Dispatch);
 
 	/* Begin IRemoteSessionChannel implementation */
-	static FString StaticType() { return TEXT("rs.xrtracking"); }
-	virtual FString GetType() const override { return StaticType(); }
+	static const TCHAR* StaticType() { return TEXT("FRemoteSessionXRTrackingChannel"); }
+	virtual const TCHAR* GetType() const override { return StaticType(); }
 	/* End IRemoteSessionChannel implementation */
 
 protected:
diff --git a/Engine/Plugins/Experimental/RemoteSession/Source/RemoteSession/Private/Channels/RemoteSessionARCameraChannel.cpp b/Engine/Plugins/Experimental/RemoteSession/Source/RemoteSession/Private/Channels/RemoteSessionARCameraChannel.cpp
--- a/Engine/Plugins/Experimental/RemoteSession/Source/RemoteSession/Private/Channels/RemoteSessionARCameraChannel.cpp
+++ b/Engine/Plugins/Experimental/RemoteSession/Source/RemoteSession/Private/Channels/RemoteSessionARCameraChannel.cpp
@@ -27,7 +27,11 @@
 #include "PostProcessParameters.h"
 #include "EngineModule.h"
 
+#include "GeneralProjectSettings.h"
 #include "ARTextures.h"
+#include "ARSessionConfig.h"
+#include "ARBlueprintLibrary.h"
+
 #include "ARBlueprintLibrary.h"
 #include "IAppleImageUtilsPlugin.h"
 
@@ -373,6 +377,18 @@ FRemoteSessionARCameraChannel::FRemoteSessionARCameraChannel(ERemoteSessionChann
 	, Connection(InConnection)
 	, Role(InRole)
 {
+	
+	static bool OnceTimeARInit = false;
+	
+	if (!OnceTimeARInit && InRole == ERemoteSessionChannelMode::Send)
+	{
+		// Initialize AR if desired - todo, does this need to check device caps?
+		UARSessionConfig* Settings = NewObject<UARSessionConfig>();
+		UARBlueprintLibrary::StartARSession(Settings);
+		
+		OnceTimeARInit = true;
+	}
+	
 	RenderingTextures[0] = nullptr;
 	RenderingTextures[1] = nullptr;
 	PPMaterial = LoadObject<UMaterialInterface>(nullptr, TEXT("/RemoteSession/ARCameraPostProcess.ARCameraPostProcess"));
@@ -389,9 +405,9 @@ FRemoteSessionARCameraChannel::FRemoteSessionARCameraChannel(ERemoteSessionChann
 		// Create our image renderer
 		SceneViewExtension = FSceneViewExtensions::NewExtension<FARCameraSceneViewExtension>(*this);
 
-		MessageCallbackHandle = Connection->GetDispatchMap().GetAddressHandler(CAMERA_MESSAGE_ADDRESS).AddRaw(this, &FRemoteSessionARCameraChannel::ReceiveARCameraImage);
+		auto Delegate = FBackChannelDispatchDelegate::FDelegate::CreateRaw(this, &FRemoteSessionARCameraChannel::ReceiveARCameraImage);
+		MessageCallbackHandle = Connection->AddMessageHandler(CAMERA_MESSAGE_ADDRESS, Delegate);
 		Connection->SetMessageOptions(CAMERA_MESSAGE_ADDRESS, 1);
-
 	}
 }
 
@@ -400,8 +416,7 @@ FRemoteSessionARCameraChannel::~FRemoteSessionARCameraChannel()
 	if (Role == ERemoteSessionChannelMode::Receive)
 	{
 		// Remove the callback so it doesn't call back on an invalid this
-		Connection->GetDispatchMap().GetAddressHandler(CAMERA_MESSAGE_ADDRESS).Remove(MessageCallbackHandle);
-		MessageCallbackHandle.Reset();
+		Connection->RemoveMessageHandler(CAMERA_MESSAGE_ADDRESS, MessageCallbackHandle);
 	}
 }
 
@@ -445,6 +460,7 @@ void FRemoteSessionARCameraChannel::QueueARCameraImage()
 	{
 		return;
 	}
+
 	UARTextureCameraImage* CameraImage = UARBlueprintLibrary::GetCameraImage();
 	if (CameraImage != nullptr && CameraImage->Timestamp > LastQueuedTimestamp)
 	{
@@ -596,4 +612,4 @@ void FRemoteSessionARCameraChannel::UpdateRenderingTexture()
 			delete InRegions;
 		});
 	} //-V773
-}
\ No newline at end of file
+}
diff --git a/Engine/Plugins/Experimental/RemoteSession/Source/RemoteSession/Private/Channels/RemoteSessionFrameBufferChannel.cpp b/Engine/Plugins/Experimental/RemoteSession/Source/RemoteSession/Private/Channels/RemoteSessionFrameBufferChannel.cpp
--- a/Engine/Plugins/Experimental/RemoteSession/Source/RemoteSession/Private/Channels/RemoteSessionFrameBufferChannel.cpp
+++ b/Engine/Plugins/Experimental/RemoteSession/Source/RemoteSession/Private/Channels/RemoteSessionFrameBufferChannel.cpp
@@ -46,10 +46,13 @@ FRemoteSessionFrameBufferChannel::FRemoteSessionFrameBufferChannel(ERemoteSessio
 	NumSentImages = 0;
 	KickedTaskCount = 0;
 	Role = InRole;
+	ViewportResized = false;
 
 	if (Role == ERemoteSessionChannelMode::Receive)
 	{
-		MessageCallbackHandle = InConnection->GetDispatchMap().GetAddressHandler(TEXT("/Screen")).AddRaw(this, &FRemoteSessionFrameBufferChannel::ReceiveHostImage);
+		auto Delegate = FBackChannelDispatchDelegate::FDelegate::CreateRaw(this, &FRemoteSessionFrameBufferChannel::ReceiveHostImage);
+		MessageCallbackHandle = InConnection->AddMessageHandler(TEXT("/Screen"), Delegate);
+
 		InConnection->SetMessageOptions(TEXT("/Screen"), 1);
 	}
 }
@@ -62,7 +65,7 @@ FRemoteSessionFrameBufferChannel::~FRemoteSessionFrameBufferChannel()
 		if (LocalConnection.IsValid())
 		{
 			// Remove the callback so it doesn't call back on an invalid this
-			LocalConnection->GetDispatchMap().GetAddressHandler(TEXT("/Screen")).Remove(MessageCallbackHandle);
+			LocalConnection->RemoveMessageHandler(TEXT("/Screen"), MessageCallbackHandle);
 		}
 		MessageCallbackHandle.Reset();
 	}
@@ -72,11 +75,7 @@ FRemoteSessionFrameBufferChannel::~FRemoteSessionFrameBufferChannel()
 		FPlatformProcess::SleepNoStats(0);
 	}
 
-	if (FrameGrabber.IsValid())
-	{
-		FrameGrabber->StopCapturingFrames();
-		FrameGrabber = nullptr;
-	}
+	ReleaseFrameGrabber();
 
 	for (int32 i = 0; i < 2; i++)
 	{
@@ -88,11 +87,16 @@ FRemoteSessionFrameBufferChannel::~FRemoteSessionFrameBufferChannel()
 	}
 }
 
-FString FRemoteSessionFrameBufferChannel::StaticType()
+void FRemoteSessionFrameBufferChannel::ReleaseFrameGrabber()
 {
-	return TEXT("rs.framebuffer");
+	if (FrameGrabber.IsValid())
+	{
+		FrameGrabber->StopCapturingFrames();
+		FrameGrabber = nullptr;
+	}
 }
 
+
 void FRemoteSessionFrameBufferChannel::SetCaptureQuality(int32 InQuality, int32 InFramerate)
 {
 	// Set our framerate and quality cvars, if the user hasn't modified them
@@ -109,6 +113,17 @@ void FRemoteSessionFrameBufferChannel::SetCaptureQuality(int32 InQuality, int32
 
 void FRemoteSessionFrameBufferChannel::SetCaptureViewport(TSharedRef<FSceneViewport> Viewport)
 {
+	SceneViewport = Viewport;
+
+	CreateFrameGrabber(Viewport);
+
+	// set the listener for the window resize event
+	Viewport->SetOnSceneViewportResizeDel(FOnSceneViewportResize::CreateRaw(this, &FRemoteSessionFrameBufferChannel::OnViewportResized));
+}
+
+void FRemoteSessionFrameBufferChannel::CreateFrameGrabber(TSharedRef<FSceneViewport> Viewport)
+{
+	ReleaseFrameGrabber();
 	FrameGrabber = MakeShareable(new FFrameGrabber(Viewport, Viewport->GetSize()));
 	FrameGrabber->StartCapturingFrames();
 }
@@ -124,6 +139,11 @@ void FRemoteSessionFrameBufferChannel::Tick(const float InDeltaTime)
 
 	if (FrameGrabber.IsValid())
 	{
+		if (ViewportResized)
+		{
+			CreateFrameGrabber(SceneViewport.ToSharedRef());
+			ViewportResized = false;
+		}
 		SCOPE_CYCLE_COUNTER(STAT_FrameBufferCapture);
 
 		FrameGrabber->CaptureThisFrame(FFramePayloadPtr());
@@ -203,9 +223,9 @@ void FRemoteSessionFrameBufferChannel::Tick(const float InDeltaTime)
 			FUpdateTextureRegion2D* Region = new FUpdateTextureRegion2D(0, 0, 0, 0, QueuedImage->Width, QueuedImage->Height);
 			TArray<uint8>* TextureData = new TArray<uint8>(MoveTemp(QueuedImage->ImageData));
 
-			DecodedTextures[NextImage]->UpdateTextureRegions(0, 1, Region, 4 * QueuedImage->Width, 8, TextureData->GetData(), [this, NextImage](auto InTextureData, auto InRegions) {
+			DecodedTextures[NextImage]->UpdateTextureRegions(0, 1, Region, 4 * QueuedImage->Width, 8, TextureData->GetData(), [this, NextImage, TextureData](auto InTextureData, auto InRegions) {
 				DecodedTextureIndex = NextImage;
-				delete InTextureData;
+				delete TextureData; // delete array, not underlying data that UpdateTextureRegions passes us
 				delete InRegions;
 			});
 
@@ -358,4 +378,8 @@ void FRemoteSessionFrameBufferChannel::CreateTexture(const int32 InSlot, const i
 	UE_LOG(LogRemoteSession, Log, TEXT("Created texture in slot %d %dx%d for incoming image"), InSlot, InWidth, InHeight);
 }
 
+void FRemoteSessionFrameBufferChannel::OnViewportResized(FVector2D NewSize)
+{
+	ViewportResized = true;
+}
 
diff --git a/Engine/Plugins/Experimental/RemoteSession/Source/RemoteSession/Private/Channels/RemoteSessionInputChannel.cpp b/Engine/Plugins/Experimental/RemoteSession/Source/RemoteSession/Private/Channels/RemoteSessionInputChannel.cpp
--- a/Engine/Plugins/Experimental/RemoteSession/Source/RemoteSession/Private/Channels/RemoteSessionInputChannel.cpp
+++ b/Engine/Plugins/Experimental/RemoteSession/Source/RemoteSession/Private/Channels/RemoteSessionInputChannel.cpp
@@ -32,10 +32,10 @@ FRemoteSessionInputChannel::FRemoteSessionInputChannel(ERemoteSessionChannelMode
 		TSharedRef<FGenericApplicationMessageHandler> DestinationHandler = FSlateApplication::Get().GetPlatformApplication()->GetMessageHandler();
 
 		PlaybackHandler = MakeShareable(new FRecordingMessageHandler(DestinationHandler));
-
-		MessageCallbackHandle = Connection->GetDispatchMap().GetAddressHandler(TEXT("/MessageHandler/")).AddRaw(this, &FRemoteSessionInputChannel::OnRemoteMessage);
+		
+		auto Delegate = FBackChannelDispatchDelegate::FDelegate::CreateRaw(this, &FRemoteSessionInputChannel::OnRemoteMessage);
+		MessageCallbackHandle = Connection->AddMessageHandler(TEXT("/MessageHandler/"), Delegate);
 	}
-	
 }
 
 FRemoteSessionInputChannel::~FRemoteSessionInputChannel()
@@ -43,7 +43,7 @@ FRemoteSessionInputChannel::~FRemoteSessionInputChannel()
 	if (Role == ERemoteSessionChannelMode::Receive)
 	{
 		// Remove the callback so it doesn't call back on an invalid this
-		Connection->GetDispatchMap().GetAddressHandler(TEXT("/MessageHandler/")).Remove(MessageCallbackHandle);
+		Connection->RemoveMessageHandler(TEXT("/MessageHandler/"), MessageCallbackHandle);
 		MessageCallbackHandle.Reset();
 	}
 
@@ -60,17 +60,20 @@ FRemoteSessionInputChannel::~FRemoteSessionInputChannel()
 	}
 }
 
-FString FRemoteSessionInputChannel::StaticType()
-{
-	return TEXT("rs.input");
-}
-
 
 void FRemoteSessionInputChannel::SetPlaybackWindow(TWeakPtr<SWindow> InWindow, TWeakPtr<FSceneViewport> InViewport)
 {
 	PlaybackHandler->SetPlaybackWindow(InWindow, InViewport);
 }
 
+void FRemoteSessionInputChannel::SetInputRect(const FVector2D& TopLeft, const FVector2D& Extents)
+{
+	if (RecordingHandler.IsValid())
+	{
+		RecordingHandler->SetInputRect(TopLeft, Extents);
+	}
+}
+
 void FRemoteSessionInputChannel::Tick(const float InDeltaTime)
 {
 	// everything happens via messaging.
diff --git a/Engine/Plugins/Experimental/RemoteSession/Source/RemoteSession/Private/Channels/RemoteSessionXRTrackingChannel.cpp b/Engine/Plugins/Experimental/RemoteSession/Source/RemoteSession/Private/Channels/RemoteSessionXRTrackingChannel.cpp
--- a/Engine/Plugins/Experimental/RemoteSession/Source/RemoteSession/Private/Channels/RemoteSessionXRTrackingChannel.cpp
+++ b/Engine/Plugins/Experimental/RemoteSession/Source/RemoteSession/Private/Channels/RemoteSessionXRTrackingChannel.cpp
@@ -47,8 +47,9 @@ FRemoteSessionXRTrackingChannel::FRemoteSessionXRTrackingChannel(ERemoteSessionC
 		// Make the proxy and set GEngine->XRSystem to it
 		ProxyXRSystem = MakeShared<FXRTrackingProxy, ESPMode::ThreadSafe>();
 		GEngine->XRSystem = ProxyXRSystem;
-
-		MessageCallbackHandle = Connection->GetDispatchMap().GetAddressHandler(MESSAGE_ADDRESS).AddRaw(this, &FRemoteSessionXRTrackingChannel::ReceiveXRTracking);
+		
+		auto Delegate = FBackChannelDispatchDelegate::FDelegate::CreateRaw(this, &FRemoteSessionXRTrackingChannel::ReceiveXRTracking);
+		MessageCallbackHandle = Connection->AddMessageHandler(MESSAGE_ADDRESS, Delegate);
 		Connection->SetMessageOptions(MESSAGE_ADDRESS, 1);
 	}
 }
@@ -58,8 +59,7 @@ FRemoteSessionXRTrackingChannel::~FRemoteSessionXRTrackingChannel()
 	if (Role == ERemoteSessionChannelMode::Receive)
 	{
 		// Remove the callback so it doesn't call back on an invalid this
-		Connection->GetDispatchMap().GetAddressHandler(MESSAGE_ADDRESS).Remove(MessageCallbackHandle);
-		MessageCallbackHandle.Reset();
+		Connection->RemoveMessageHandler(MESSAGE_ADDRESS, MessageCallbackHandle);
 
         if (GEngine != nullptr)
         {
diff --git a/Engine/Plugins/Experimental/RemoteSession/Source/RemoteSession/Private/MessageHandler/ProxyMessageHandler.cpp b/Engine/Plugins/Experimental/RemoteSession/Source/RemoteSession/Private/MessageHandler/ProxyMessageHandler.cpp
--- a/Engine/Plugins/Experimental/RemoteSession/Source/RemoteSession/Private/MessageHandler/ProxyMessageHandler.cpp
+++ b/Engine/Plugins/Experimental/RemoteSession/Source/RemoteSession/Private/MessageHandler/ProxyMessageHandler.cpp
@@ -198,21 +198,37 @@ void FProxyMessageHandler::OnEndGesture()
 	}
 }
 
+#if REMOTE_WITH_FORCE_PARAM
 bool FProxyMessageHandler::OnTouchStarted(const TSharedPtr< FGenericWindow >& Window, const FVector2D& Location, float Force, int32 TouchIndex, int32 ControllerId)
+#else
+bool FProxyMessageHandler::OnTouchStarted(const TSharedPtr< FGenericWindow >& Window, const FVector2D& Location, int32 TouchIndex, int32 ControllerId)
+#endif // REMOTE_WITH_FORCE_PARAM
 {
 	if (TargetHandler.IsValid())
 	{
+#if REMOTE_WITH_FORCE_PARAM
 		return TargetHandler->OnTouchStarted(Window, Location, Force, TouchIndex, ControllerId);
+#else
+		return TargetHandler->OnTouchStarted(Window, Location, TouchIndex, ControllerId);
+#endif
 	}
 
 	return false;
 }
 
+#if REMOTE_WITH_FORCE_PARAM
 bool FProxyMessageHandler::OnTouchMoved(const FVector2D& Location, float Force, int32 TouchIndex, int32 ControllerId)
+#else
+bool FProxyMessageHandler::OnTouchMoved(const FVector2D& Location, int32 TouchIndex, int32 ControllerId)
+#endif // REMOTE_WITH_FORCE_PARAM
 {
 	if (TargetHandler.IsValid())
 	{
+#if REMOTE_WITH_FORCE_PARAM
 		return TargetHandler->OnTouchMoved(Location, Force, TouchIndex, ControllerId);
+#else
+		return TargetHandler->OnTouchMoved(Location, TouchIndex, ControllerId);
+#endif
 	}
 
 	return false;
diff --git a/Engine/Plugins/Experimental/RemoteSession/Source/RemoteSession/Private/MessageHandler/ProxyMessageHandler.h b/Engine/Plugins/Experimental/RemoteSession/Source/RemoteSession/Private/MessageHandler/ProxyMessageHandler.h
--- a/Engine/Plugins/Experimental/RemoteSession/Source/RemoteSession/Private/MessageHandler/ProxyMessageHandler.h
+++ b/Engine/Plugins/Experimental/RemoteSession/Source/RemoteSession/Private/MessageHandler/ProxyMessageHandler.h
@@ -3,8 +3,14 @@
 #pragma once
 
 #include "CoreMinimal.h"
+#include "Launch/Resources/Version.h"
 #include "GenericPlatform/GenericApplicationMessageHandler.h"
 
+// For compatibility with a change that has not yet made it to //UE4/Main.
+// Please let andrew.grant know when you run into this as we may want to do some
+// cleanup if it occurs prior to the 4.20 branch.
+#define REMOTE_WITH_FORCE_PARAM 0
+
 class FProxyMessageHandler : public FGenericApplicationMessageHandler
 {
 	
@@ -56,9 +62,15 @@ public:
 
 	virtual void OnEndGesture() override;
 
+#if REMOTE_WITH_FORCE_PARAM
 	virtual bool OnTouchStarted(const TSharedPtr< FGenericWindow >& Window, const FVector2D& Location, float Force, int32 TouchIndex, int32 ControllerId) override;
 
 	virtual bool OnTouchMoved(const FVector2D& Location, float Force, int32 TouchIndex, int32 ControllerId) override;
+#else
+	virtual bool OnTouchStarted(const TSharedPtr< FGenericWindow >& Window, const FVector2D& Location, int32 TouchIndex, int32 ControllerId) override;
+
+	virtual bool OnTouchMoved(const FVector2D& Location, int32 TouchIndex, int32 ControllerId) override;
+#endif // REMOTE_WITH_FORCE_PARAM
 
 	virtual bool OnTouchEnded(const FVector2D& Location, int32 TouchIndex, int32 ControllerId) override;
 
diff --git a/Engine/Plugins/Experimental/RemoteSession/Source/RemoteSession/Private/MessageHandler/RecordingMessageHandler.cpp b/Engine/Plugins/Experimental/RemoteSession/Source/RemoteSession/Private/MessageHandler/RecordingMessageHandler.cpp
--- a/Engine/Plugins/Experimental/RemoteSession/Source/RemoteSession/Private/MessageHandler/RecordingMessageHandler.cpp
+++ b/Engine/Plugins/Experimental/RemoteSession/Source/RemoteSession/Private/MessageHandler/RecordingMessageHandler.cpp
@@ -1,4 +1,4 @@
-// Copyright 1998-2018 Epic Games, Inc. All Rights Reserved.
+﻿// Copyright 1998-2018 Epic Games, Inc. All Rights Reserved.
 
 #include "RecordingMessageHandler.h"
 #include "RemoteSession.h"
@@ -8,6 +8,12 @@
 #include "Engine/GameViewportClient.h"
 #include "Async/Async.h"
 
+#include "Components/CanvasPanelSlot.h"
+#include "Blueprint/WidgetLayoutLibrary.h"
+#include "Slate/SceneViewport.h"
+#include "Layout/WidgetPath.h"
+#include "Framework/Application/SlateApplication.h"
+
 // helper to serialize out const params
 template <typename S, typename T>
 S& SerializeOut(S& Ar, const T& Value)
@@ -32,6 +38,10 @@ FRecordingMessageHandler::FRecordingMessageHandler(const TSharedPtr<FGenericAppl
 {
 	OutputWriter = nullptr;
 	ConsumeInput = false;
+	bIsTouching = false;
+	InputRect = FRect(EForceInit::ForceInitToZero);
+	LastTouchLocation = FVector2D(EForceInit::ForceInitToZero);
+    
 
 	BIND_PLAYBACK_HANDLER(TEXT("OnKeyChar"), PlayOnKeyChar);
 	BIND_PLAYBACK_HANDLER(TEXT("OnKeyUp"), PlayOnKeyUp);
@@ -72,11 +82,30 @@ void FRecordingMessageHandler::SetPlaybackWindow(TWeakPtr<SWindow> InWindow, TWe
 	PlaybackViewport = InViewport;
 }
 
-FVector2D FRecordingMessageHandler::ConvertToNormalizedScreenLocation(const FVector2D& Location)
+void FRecordingMessageHandler::SetInputRect(const FVector2D& TopLeft, const FVector2D& Extents)
+{
+	InputRect = FRect(TopLeft, Extents);
+}
+
+bool FRecordingMessageHandler::ConvertToNormalizedScreenLocation(const FVector2D& InLocation, FVector2D& OutLocation)
 {
-	const FVector2D ViewportSize = FVector2D(GEngine->GameViewport->Viewport->GetSizeXY());
+	FRect ClipRect = InputRect;
+	FIntPoint Point = FIntPoint((int)InLocation.X, (int)InLocation.Y);
+
+	if (ClipRect.Width == 0 || ClipRect.Height == 0)
+	{
+		ClipRect = FRect(FVector2D(EForceInit::ForceInitToZero), FVector2D(GEngine->GameViewport->Viewport->GetSizeXY()));
+	}
+
+	if (!ClipRect.Contains(Point))
+	{
+		OutLocation = FVector2D(EForceInit::ForceInitToZero);
+		return false;
+	}
+
+	OutLocation = FVector2D((InLocation.X-ClipRect.X) / ClipRect.Width, (InLocation.Y - ClipRect.Y) / ClipRect.Height);
 
-	return FVector2D(Location.X / ViewportSize.X, Location.Y / ViewportSize.Y);
+	return true;
 }
 
 FVector2D FRecordingMessageHandler::ConvertFromNormalizedScreenLocation(const FVector2D& ScreenLocation)
@@ -84,11 +113,36 @@ FVector2D FRecordingMessageHandler::ConvertFromNormalizedScreenLocation(const FV
 	FVector2D OutVector = ScreenLocation;
 
 	TSharedPtr<SWindow> GameWindow = PlaybackWindow.Pin();
+	TSharedPtr<FSceneViewport> GameWidget = PlaybackViewport.Pin();
 	if (GameWindow.IsValid())
 	{
 		FVector2D WindowOrigin = GameWindow->GetPositionInScreen();
-		FVector2D WindowSize = GameWindow->GetSizeInScreen();
-		OutVector = WindowOrigin + (ScreenLocation * WindowSize);
+
+		TSharedPtr<SViewport> ViewportWidget = GameWidget->GetViewportWidget().Pin();
+
+		if (ViewportWidget.IsValid())
+		{
+			if (GameWindow.IsValid())
+			{
+				FGeometry InnerWindowGeometry = GameWindow->GetWindowGeometryInWindow();
+
+				// Find the widget path relative to the window
+				FArrangedChildren JustWindow(EVisibility::Visible);
+				JustWindow.AddWidget(FArrangedWidget(GameWindow.ToSharedRef(), InnerWindowGeometry));
+
+				FWidgetPath WidgetPath(GameWindow.ToSharedRef(), JustWindow);
+				if (WidgetPath.ExtendPathTo(FWidgetMatcher(ViewportWidget.ToSharedRef()), EVisibility::Visible))
+				{
+					FArrangedWidget ArrangedWidget = WidgetPath.FindArrangedWidget(ViewportWidget.ToSharedRef()).Get(FArrangedWidget::NullWidget);
+
+					FVector2D WindowClientOffset = ArrangedWidget.Geometry.GetAbsolutePosition();
+					FVector2D WindowClientSize = ArrangedWidget.Geometry.GetAbsoluteSize();
+
+					OutVector = WindowOrigin + WindowClientOffset + (ScreenLocation * WindowClientSize);
+
+				}
+			}
+		}
 	}
 
 	return OutVector;
@@ -185,12 +239,28 @@ void FRecordingMessageHandler::PlayOnKeyUp(FArchive& Ar)
 	OnKeyUp(Msg.Param1, Msg.Param2, Msg.Param3);
 }
 
+#if REMOTE_WITH_FORCE_PARAM
 bool FRecordingMessageHandler::OnTouchStarted(const TSharedPtr< FGenericWindow >& Window, const FVector2D& Location, float Force, int32 TouchIndex, int32 ControllerId)
+#else
+bool FRecordingMessageHandler::OnTouchStarted(const TSharedPtr< FGenericWindow >& Window, const FVector2D& Location, int32 TouchIndex, int32 ControllerId)
+#endif
 {
 	if (IsRecording())
 	{
-		FourParamMsg<FVector2D, float, int32, int32> Msg(ConvertToNormalizedScreenLocation(Location), Force, TouchIndex, ControllerId);
-		RecordMessage(TEXT("OnTouchStarted"), Msg.AsData());
+		FVector2D Normalized;
+
+#if !REMOTE_WITH_FORCE_PARAM
+		float Force = 1.0f;
+#endif
+
+		if (ConvertToNormalizedScreenLocation(Location, Normalized))
+		{
+			// note - force is serialized last for backwards compat - force was introduced in 4.20
+			FourParamMsg<FVector2D, int32, int32, float> Msg(Normalized, TouchIndex, ControllerId, Force);
+			RecordMessage(TEXT("OnTouchStarted"), Msg.AsData());
+			bIsTouching = true;
+			LastTouchLocation = Location;
+		}
 	}
 
 	if (ConsumeInput)
@@ -198,12 +268,16 @@ bool FRecordingMessageHandler::OnTouchStarted(const TSharedPtr< FGenericWindow >
 		return true;
 	}
 
+#if REMOTE_WITH_FORCE_PARAM
 	return FProxyMessageHandler::OnTouchStarted(Window, Location, Force, TouchIndex, ControllerId);
+#else
+	return FProxyMessageHandler::OnTouchStarted(Window, Location, TouchIndex, ControllerId);
+#endif
 }
 
 void FRecordingMessageHandler::PlayOnTouchStarted(FArchive& Ar)
 {
-	FourParamMsg<FVector2D, float, int32, int32 > Msg(Ar);
+	FourParamMsg<FVector2D, int32, int32, float > Msg(Ar);
 	FVector2D ScreenLocation = ConvertFromNormalizedScreenLocation(Msg.Param1);
 
 	TSharedPtr<FGenericWindow> Window;
@@ -213,15 +287,35 @@ void FRecordingMessageHandler::PlayOnTouchStarted(FArchive& Ar)
 		Window = PlaybackWindow.Pin()->GetNativeWindow();
 	}
 
-	OnTouchStarted(Window, ScreenLocation, Msg.Param2, Msg.Param3, Msg.Param4);
+#if REMOTE_WITH_FORCE_PARAM
+	// note - force is serialized last for backwards compat - force was introduced in 4.20
+	OnTouchStarted(Window, ScreenLocation, Msg.Param4, Msg.Param2, Msg.Param3 );
+#else
+	OnTouchStarted(Window, ScreenLocation, Msg.Param2, Msg.Param3);
+#endif
 }
 
+#if REMOTE_WITH_FORCE_PARAM
 bool FRecordingMessageHandler::OnTouchMoved(const FVector2D& Location, float Force, int32 TouchIndex, int32 ControllerId)
+#else
+bool FRecordingMessageHandler::OnTouchMoved(const FVector2D& Location, int32 TouchIndex, int32 ControllerId)
+#endif
 {
 	if (IsRecording())
 	{
-		FourParamMsg<FVector2D, float, int32, int32> Msg(ConvertToNormalizedScreenLocation(Location), Force, TouchIndex, ControllerId);
-		OutputWriter->RecordMessage(TEXT("OnTouchMoved"), Msg.AsData());
+		FVector2D Normalized;
+
+#if !REMOTE_WITH_FORCE_PARAM
+		float Force = 1.0f;
+#endif
+		if (ConvertToNormalizedScreenLocation(Location, Normalized))
+		{
+			// note - force is serialized last for backwards compat - force was introduced in 4.20
+			FourParamMsg<FVector2D, int32, int32, float> Msg(Normalized, TouchIndex, ControllerId, Force);
+			OutputWriter->RecordMessage(TEXT("OnTouchMoved"), Msg.AsData());
+			bIsTouching = true;
+			LastTouchLocation = Location;
+		}
 	}
 
 	if (ConsumeInput)
@@ -229,22 +323,40 @@ bool FRecordingMessageHandler::OnTouchMoved(const FVector2D& Location, float For
 		return true;
 	}
 
+#if REMOTE_WITH_FORCE_PARAM
 	return FProxyMessageHandler::OnTouchMoved(Location, Force, TouchIndex, ControllerId);
+#else
+	return FProxyMessageHandler::OnTouchMoved(Location, TouchIndex, ControllerId);
+#endif
 }
 
 void FRecordingMessageHandler::PlayOnTouchMoved(FArchive& Ar)
 {
-	FourParamMsg<FVector2D, float, int32, int32 > Msg(Ar);
+	FourParamMsg<FVector2D, int32, int32, float > Msg(Ar);
 	FVector2D ScreenLocation = ConvertFromNormalizedScreenLocation(Msg.Param1);
-	OnTouchMoved(ScreenLocation, Msg.Param2, Msg.Param3, Msg.Param4);
+#if REMOTE_WITH_FORCE_PARAM
+	// note - force is serialized last for backwards compat - force was introduced in 4.20
+	OnTouchMoved(ScreenLocation, Msg.Param4, Msg.Param2, Msg.Param3);
+#else
+	OnTouchMoved(ScreenLocation, Msg.Param2, Msg.Param3);
+#endif
 }
 
 bool FRecordingMessageHandler::OnTouchEnded(const FVector2D& Location, int32 TouchIndex, int32 ControllerId)
 {
 	if (IsRecording())
 	{
-		ThreeParamMsg<FVector2D, int32, int32> Msg(ConvertToNormalizedScreenLocation(Location), TouchIndex, ControllerId);
+		FVector2D Normalized;
+
+		// if outside our bounds, end the touch where it left
+		if (ConvertToNormalizedScreenLocation(Location, Normalized) == false)
+		{
+			ConvertToNormalizedScreenLocation(LastTouchLocation, Normalized);
+		}
+		
+		ThreeParamMsg<FVector2D, int32, int32> Msg(Normalized, TouchIndex, ControllerId);
 		OutputWriter->RecordMessage(TEXT("OnTouchEnded"), Msg.AsData());
+		bIsTouching = false;
 	}
 
 	if (ConsumeInput)
diff --git a/Engine/Plugins/Experimental/RemoteSession/Source/RemoteSession/Private/RemoteSessionClient.cpp b/Engine/Plugins/Experimental/RemoteSession/Source/RemoteSession/Private/RemoteSessionClient.cpp
--- a/Engine/Plugins/Experimental/RemoteSession/Source/RemoteSession/Private/RemoteSessionClient.cpp
+++ b/Engine/Plugins/Experimental/RemoteSession/Source/RemoteSession/Private/RemoteSessionClient.cpp
@@ -2,6 +2,7 @@
 
 #include "RemoteSessionClient.h"
 #include "BackChannel/Transport/IBackChannelTransport.h"
+#include "BackChannel/Protocol/OSC/BackChannelOSCMessage.h"
 #include "Framework/Application/SlateApplication.h"	
 #include "Channels/RemoteSessionInputChannel.h"
 #include "Channels/RemoteSessionXRTrackingChannel.h"
@@ -11,10 +12,12 @@
 #include "IImageWrapperModule.h"
 #include "Sockets.h"
 #include "RemoteSession.h"
-
+#include "Async/Async.h"
 
 DECLARE_CYCLE_STAT(TEXT("RSClientTick"), STAT_RDClientTick, STATGROUP_Game);
 
+#define RS_TIMEOUT_IS_ERROR 0
+
 FRemoteSessionClient::FRemoteSessionClient(const TCHAR* InHostAddress)
 {
 	HostAddress = InHostAddress;
@@ -50,7 +53,7 @@ void FRemoteSessionClient::Tick(float DeltaTime)
 
 	if (IsConnected() == false)
 	{
-		if (IsConnecting == false)
+		if (IsConnecting == false && HasError()  == false)
 		{
 			const double TimeSinceLastAttempt = FPlatformTime::Seconds() - TimeConnectionAttemptStarted;
 
@@ -105,20 +108,12 @@ void FRemoteSessionClient::CheckConnection()
 
 		Connection->GetSocket()->SetReceiveBufferSize(WantedSize, ActualSize);
 
-		OSCConnection = MakeShareable(new FBackChannelOSCConnection(Connection.ToSharedRef()));
-
-		Channels.Add(MakeShareable(new FRemoteSessionInputChannel(ERemoteSessionChannelMode::Send, OSCConnection)));
-		Channels.Add(MakeShareable(new FRemoteSessionXRTrackingChannel(ERemoteSessionChannelMode::Send, OSCConnection)));
-#if PLATFORM_IOS // Client side sending only works on iOS with Android coming in the future
-		Channels.Add(MakeShareable(new FRemoteSessionARCameraChannel(ERemoteSessionChannelMode::Send, OSCConnection)));
-#endif
-		Channels.Add(MakeShareable(new FRemoteSessionFrameBufferChannel(ERemoteSessionChannelMode::Receive, OSCConnection)));
-
+		CreateOSCConnection(Connection.ToSharedRef());
+		
 		UE_LOG(LogRemoteSession, Log, TEXT("Connected to host at %s (ReceiveSize=%dkb)"), *HostAddress, ActualSize / 1024);
 
 		IsConnecting = false;
 
-		OSCConnection->StartReceiveThread();
 		//SetReceiveInBackground(true);
 
 		return true;
@@ -131,17 +126,65 @@ void FRemoteSessionClient::CheckConnection()
 		if (Success == false || TimeSpentConnecting >= ConnectionTimeout)
 		{
 			IsConnecting = false;
+			
+			FString Msg;
+			
 			if (TimeSpentConnecting >= ConnectionTimeout)
 			{
-				UE_LOG(LogRemoteSession, Log, TEXT("Timing out connection attempt after %.02f seconds"), TimeSpentConnecting);	
+				Msg = FString::Printf(TEXT("Timing out connection attempt after %.02f seconds"), TimeSpentConnecting);
 			}
 			else
 			{
-				UE_LOG(LogRemoteSession, Log, TEXT("Failed to check for connection. Aborting."));
+				Msg = TEXT("Failed to check for connection. Aborting.");
 			}
-
+			
+			UE_LOG(LogRemoteSession, Log, TEXT("%s"), *Msg);
+			
+#if RS_TIMEOUT_IS_ERROR
+			CloseWithError(*Msg);
+#else
 			Close();
+#endif
 			TimeConnectionAttemptStarted = FPlatformTime::Seconds();
 		}
 	}
 }
+
+void FRemoteSessionClient::OnBindEndpoints()
+{
+	auto Delegate = FBackChannelDispatchDelegate::FDelegate::CreateRaw(this, &FRemoteSessionClient::OnChannelSelection);
+	OSCConnection->AddMessageHandler(*GetChannelSelectionEndPoint(),Delegate);
+	
+	FRemoteSessionRole::OnBindEndpoints();
+}
+
+void FRemoteSessionClient::OnChannelSelection(FBackChannelOSCMessage& Message, FBackChannelOSCDispatch& Dispatch)
+{
+	TMap<FString, ERemoteSessionChannelMode> DesiredChannels;
+	
+	const int NumChannels = Message.GetArgumentCount() / 2;
+	
+	for (int i = 0; i < NumChannels; i++)
+	{
+		FString ChannelName;
+		int32 ChannelMode;
+		
+		Message.Read(ChannelName);
+		Message.Read(ChannelMode);
+		
+		if (ChannelName.Len())
+		{
+			DesiredChannels.Add(ChannelName, (ERemoteSessionChannelMode)ChannelMode);
+		}
+		else
+		{
+			UE_LOG(LogRemoteSession, Error, TEXT("Failed to read channel from ChannelSelection message!"));
+		}
+	}
+	
+	// Need to create channels on the main thread
+	AsyncTask(ENamedThreads::GameThread, [this, DesiredChannels]
+	{
+		CreateChannels(DesiredChannels);
+	});
+}
diff --git a/Engine/Plugins/Experimental/RemoteSession/Source/RemoteSession/Private/RemoteSessionClient.h b/Engine/Plugins/Experimental/RemoteSession/Source/RemoteSession/Private/RemoteSessionClient.h
--- a/Engine/Plugins/Experimental/RemoteSession/Source/RemoteSession/Private/RemoteSessionClient.h
+++ b/Engine/Plugins/Experimental/RemoteSession/Source/RemoteSession/Private/RemoteSessionClient.h
@@ -4,6 +4,8 @@
 
 #include "RemoteSessionRole.h"
 
+class FBackChannelOSCMessage;
+class FBackChannelOSCDispatch;
 
 class FRemoteSessionClient : public FRemoteSessionRole
 {
@@ -18,8 +20,12 @@ public:
 
 protected:
 
-	void StartConnection();
-	void CheckConnection();
+	virtual void	OnBindEndpoints() override;
+	
+	void 			StartConnection();
+	void 			CheckConnection();
+	
+	void OnChannelSelection(FBackChannelOSCMessage& Message, FBackChannelOSCDispatch& Dispatch);
 
 	FString				HostAddress;
 	
@@ -28,4 +34,6 @@ protected:
 
 	double				ConnectionAttemptTimer;
 	double				TimeConnectionAttemptStarted;
+	
+	FDelegateHandle		ChannelCallbackHandle;
 };
diff --git a/Engine/Plugins/Experimental/RemoteSession/Source/RemoteSession/Private/RemoteSessionHost.cpp b/Engine/Plugins/Experimental/RemoteSession/Source/RemoteSession/Private/RemoteSessionHost.cpp
--- a/Engine/Plugins/Experimental/RemoteSession/Source/RemoteSession/Private/RemoteSessionHost.cpp
+++ b/Engine/Plugins/Experimental/RemoteSession/Source/RemoteSession/Private/RemoteSessionHost.cpp
@@ -1,29 +1,39 @@
-// Copyright 1998-2018 Epic Games, Inc. All Rights Reserved.
+﻿// Copyright 1998-2018 Epic Games, Inc. All Rights Reserved.
 
 #include "RemoteSessionHost.h"
 #include "BackChannel/Transport/IBackChannelTransport.h"
 #include "FrameGrabber.h"
 #include "Widgets/SViewport.h"
 #include "BackChannel/Utils/BackChannelThreadedConnection.h"
+#include "BackChannel/Protocol/OSC/BackChannelOSCMessage.h"
 #include "Channels/RemoteSessionInputChannel.h"
 #include "Channels/RemoteSessionXRTrackingChannel.h"
 #include "Channels/RemoteSessionARCameraChannel.h"
 #include "Channels/RemoteSessionFrameBufferChannel.h"
 #include "Engine/GameEngine.h"
+#include "Framework/Application/SlateApplication.h"
+#include "Misc/ConfigCacheIni.h"
+#include "RemoteSession.h"
 
 #if WITH_EDITOR
 	#include "Editor.h"
 	#include "Editor/EditorEngine.h"
+	#include "ILevelViewport.h"
 #endif
 
-PRAGMA_DISABLE_OPTIMIZATION
-
+namespace RemoteSessionEd
+{
+	static FAutoConsoleVariable SlateDragDistanceOverride(TEXT("RemoteSessionEd.SlateDragDistanceOverride"), 10.0f, TEXT("How many pixels you need to drag before a drag and drop operation starts in remote app"));
+};
 
 
 FRemoteSessionHost::FRemoteSessionHost(int32 InQuality, int32 InFramerate)
 {
+	HostTCPPort = 0;
 	Quality = InQuality;
 	Framerate = InFramerate;
+	SavedEditorDragTriggerDistance = FSlateApplication::Get().GetDragTriggerDistance();
+	IsListenerConnected = false;
 }
 
 FRemoteSessionHost::~FRemoteSessionHost()
@@ -38,6 +48,17 @@ FRemoteSessionHost::~FRemoteSessionHost()
 	Close();
 }
 
+
+void FRemoteSessionHost::Close()
+{
+	FRemoteSessionRole::Close();
+
+	if (FSlateApplication::IsInitialized())
+	{
+		FSlateApplication::Get().SetDragTriggerDistance(SavedEditorDragTriggerDistance);
+	}
+}
+
 void FRemoteSessionHost::SetScreenSharing(const bool bEnabled)
 {
 }
@@ -61,22 +82,60 @@ bool FRemoteSessionHost::StartListening(const uint16 InPort)
 	{
 		Listener = Transport->CreateConnection(IBackChannelTransport::TCP);
 
-		Listener->Listen(InPort);
+		if (Listener->Listen(InPort) == false)
+		{
+			Listener = nullptr;
+		}
+		HostTCPPort = InPort;
 	}
 
 	return Listener.IsValid();
 }
 
-
-bool FRemoteSessionHost::ProcessIncomingConnection(TSharedRef<IBackChannelConnection> NewConnection)
+void FRemoteSessionHost::OnBindEndpoints()
 {
-	Close();
-
-	OSCConnection = MakeShareable(new FBackChannelOSCConnection(NewConnection));
+	FRemoteSessionRole::OnBindEndpoints();
+}
 
+void FRemoteSessionHost::OnCreateChannels()
+{
+	FRemoteSessionRole::OnCreateChannels();
+	
+	ClearChannels();
+	
+	// Query the list of channels from the hosts ini file.
+	TArray<FString> DesiredChannels;
+	GConfig->GetArray(TEXT("RemoteSession"), TEXT("Channels"), DesiredChannels, GEngineIni);
+	
+	if (DesiredChannels.Num() == 0)
+	{
+		// Default to Input receive and framebuffer send
+		DesiredChannels.Add(FString::Printf(TEXT("%s,receive"), *FRemoteSessionInputChannel::StaticType()));
+		DesiredChannels.Add(FString::Printf(TEXT("%s,send"), *FRemoteSessionFrameBufferChannel::StaticType()));
+	}
+	
+	TMap<FString, ERemoteSessionChannelMode> ChannelsWithModes;
+	
+	for (const FString& Channel : DesiredChannels)
+	{
+		FString ChannelName, Mode;
+		
+		if (Channel.Split(TEXT(","), &ChannelName, &Mode))
+		{
+			ChannelsWithModes.Add(ChannelName, Mode == TEXT("receive") ? ERemoteSessionChannelMode::Receive : ERemoteSessionChannelMode::Send);
+		}
+		else
+		{
+			UE_LOG(LogRemoteSession, Error, TEXT("Unrecognized channel syntax '%s'. Should be ChannelType,r or ChannelType,s"), *Channel);
+		}
+	}
+	
+	CreateChannels(ChannelsWithModes);
+	IsListenerConnected = true;
+	
 	TWeakPtr<SWindow> InputWindow;
 	TSharedPtr<FSceneViewport> SceneViewport;
-
+	
 #if WITH_EDITOR
 	if (GIsEditor)
 	{
@@ -87,16 +146,23 @@ bool FRemoteSessionHost::ProcessIncomingConnection(TSharedRef<IBackChannelConnec
 				FSlatePlayInEditorInfo* SlatePlayInEditorSession = GEditor->SlatePlayInEditorMap.Find(Context.ContextHandle);
 				if (SlatePlayInEditorSession)
 				{
-					if (SlatePlayInEditorSession->SlatePlayInEditorWindowViewport.IsValid())
+					if (SlatePlayInEditorSession->DestinationSlateViewport.IsValid())
+					{
+						TSharedPtr<ILevelViewport> DestinationLevelViewport = SlatePlayInEditorSession->DestinationSlateViewport.Pin();
+						SceneViewport = DestinationLevelViewport->GetSharedActiveViewport();
+						InputWindow = FSlateApplication::Get().FindWidgetWindow(DestinationLevelViewport->AsWidget());
+					}
+					else if (SlatePlayInEditorSession->SlatePlayInEditorWindowViewport.IsValid())
 					{
 						SceneViewport = SlatePlayInEditorSession->SlatePlayInEditorWindowViewport;
+						InputWindow = SlatePlayInEditorSession->SlatePlayInEditorWindow;
 					}
-
-					InputWindow = SlatePlayInEditorSession->SlatePlayInEditorWindow;
 				}
 			}
 		}
 		
+		SavedEditorDragTriggerDistance = FSlateApplication::Get().GetDragTriggerDistance();
+		FSlateApplication::Get().SetDragTriggerDistance(RemoteSessionEd::SlateDragDistanceOverride->GetFloat());
 	}
 	else
 #endif
@@ -105,25 +171,34 @@ bool FRemoteSessionHost::ProcessIncomingConnection(TSharedRef<IBackChannelConnec
 		SceneViewport = GameEngine->SceneViewport;
 		InputWindow = GameEngine->GameViewportWindow;
 	}
-
-	TSharedPtr<FRemoteSessionInputChannel> InputChannel = MakeShareable(new FRemoteSessionInputChannel(ERemoteSessionChannelMode::Receive, OSCConnection));
-	InputChannel->SetPlaybackWindow(InputWindow, SceneViewport);
-	Channels.Add(InputChannel);
-	Channels.Add(MakeShareable(new FRemoteSessionXRTrackingChannel(ERemoteSessionChannelMode::Receive, OSCConnection)));
-
-	if (SceneViewport.IsValid())
+	
+	// setup framebuffer capture
+	TSharedPtr<FRemoteSessionFrameBufferChannel> FBChannel = IRemoteSessionRole::GetChannel<FRemoteSessionFrameBufferChannel>();
+	if (FBChannel.IsValid())
 	{
-		Channels.Add(MakeShareable(new FRemoteSessionARCameraChannel(ERemoteSessionChannelMode::Receive, OSCConnection)));
-
-		TSharedPtr<FRemoteSessionFrameBufferChannel> FramebufferChannel = MakeShareable(new FRemoteSessionFrameBufferChannel(ERemoteSessionChannelMode::Send, OSCConnection));
-		FramebufferChannel->SetCaptureViewport(SceneViewport.ToSharedRef());
-		FramebufferChannel->SetCaptureQuality(Quality, Framerate);
-		Channels.Add(FramebufferChannel);
+		FBChannel->SetCaptureViewport(SceneViewport.ToSharedRef());
+		FBChannel->SetCaptureQuality(Quality, Framerate);
 	}
-
-	OSCConnection->StartReceiveThread();
-
-	return true;
+	
+	// setup input playback
+	TSharedPtr<FRemoteSessionInputChannel> InputChannel = IRemoteSessionRole::GetChannel<FRemoteSessionInputChannel>();
+	if (InputChannel.IsValid())
+	{
+		InputChannel->SetPlaybackWindow(InputWindow, SceneViewport);
+	}
+	
+	// now ask the client to start these channels
+	FBackChannelOSCMessage Msg(*GetChannelSelectionEndPoint());
+	
+	// send these across as a name/mode pair
+	for (const auto& KP : ChannelsWithModes)
+	{
+		ERemoteSessionChannelMode ClientMode = (KP.Value == ERemoteSessionChannelMode::Send) ? ERemoteSessionChannelMode::Receive : ERemoteSessionChannelMode::Send;
+		Msg.Write(KP.Key);
+		Msg.Write((int32)ClientMode);
+	}
+	
+	OSCConnection->SendPacket(Msg);
 }
 
 
@@ -132,14 +207,21 @@ void FRemoteSessionHost::Tick(float DeltaTime)
 	// non-threaded listener
 	if (IsConnected() == false)
 	{
-		Listener->WaitForConnection(0, [this](TSharedRef<IBackChannelConnection> NewConnection) {
-			return ProcessIncomingConnection(NewConnection);
+		if (Listener.IsValid() && IsListenerConnected)
+		{
+			Listener->Close();
+			Listener = nullptr;
+
+			//reset the host TCP socket
+			StartListening(HostTCPPort);
+			IsListenerConnected = false;
+		}
+		Listener->WaitForConnection(0, [this](TSharedRef<IBackChannelConnection> InConnection) {
+			Close();
+			CreateOSCConnection(InConnection);
+			return true;
 		});
 	}
 	
 	FRemoteSessionRole::Tick(DeltaTime);
 }
-
-
-
-PRAGMA_ENABLE_OPTIMIZATION
\ No newline at end of file
diff --git a/Engine/Plugins/Experimental/RemoteSession/Source/RemoteSession/Private/RemoteSessionHost.h b/Engine/Plugins/Experimental/RemoteSession/Source/RemoteSession/Private/RemoteSessionHost.h
--- a/Engine/Plugins/Experimental/RemoteSession/Source/RemoteSession/Private/RemoteSessionHost.h
+++ b/Engine/Plugins/Experimental/RemoteSession/Source/RemoteSession/Private/RemoteSessionHost.h
@@ -1,4 +1,4 @@
-// Copyright 1998-2018 Epic Games, Inc. All Rights Reserved.
+﻿// Copyright 1998-2018 Epic Games, Inc. All Rights Reserved.
 
 #pragma once
 
@@ -17,6 +17,8 @@ public:
 	FRemoteSessionHost(int32 InQuality, int32 InFramerate);
 	~FRemoteSessionHost();
 
+	virtual void Close() override;
+
 	bool StartListening(const uint16 Port);
 
 	void SetScreenSharing(const bool bEnabled);
@@ -27,10 +29,22 @@ public:
 
 protected:
 
-	bool	ProcessIncomingConnection(TSharedRef<IBackChannelConnection> NewConnection);
+	virtual void	OnBindEndpoints() override;
+	virtual void	OnCreateChannels() override;
+	
+	bool			ProcessIncomingConnection(TSharedRef<IBackChannelConnection> NewConnection);
 
 	TSharedPtr<IBackChannelConnection> Listener;
 
 	int32		Quality;
 	int32		Framerate;
+
+	/** Saved information about the editor and viewport we possessed, so we can restore it after exiting VR mode */
+	float SavedEditorDragTriggerDistance;
+
+	/** Host's TCP port */
+	uint16 HostTCPPort;
+
+	/** True if the host TCP socket is connected*/
+	bool IsListenerConnected;
 };
\ No newline at end of file
diff --git a/Engine/Plugins/Experimental/RemoteSession/Source/RemoteSession/Private/RemoteSessionModule.cpp b/Engine/Plugins/Experimental/RemoteSession/Source/RemoteSession/Private/RemoteSessionModule.cpp
--- a/Engine/Plugins/Experimental/RemoteSession/Source/RemoteSession/Private/RemoteSessionModule.cpp
+++ b/Engine/Plugins/Experimental/RemoteSession/Source/RemoteSession/Private/RemoteSessionModule.cpp
@@ -20,15 +20,17 @@ class FRemoteSessionModule : public IRemoteSessionModule, public FTickableGameOb
 protected:
 
 	TSharedPtr<FRemoteSessionHost>		Host;
-	TSharedPtr<FRemoteSessionClient>		Client;
+	TSharedPtr<FRemoteSessionClient>	Client;
 
 	int32								DefaultPort;
 	int32								Quality;
 	int32								Framerate;
 
 	bool bAutoHostWithPIE;
+	bool bAutoHostWithGame;
 	FDelegateHandle PostPieDelegate;
 	FDelegateHandle EndPieDelegate;
+	FDelegateHandle GameStartDelegate;
 
 public:
 
@@ -39,11 +41,12 @@ public:
 
 	void StartupModule()
 	{
-		bool bAutoHostWithGame = false;
+		// set defaults
 		DefaultPort = IRemoteSessionModule::kDefaultPort;
 		Quality = 85;
 		Framerate = 30;
-		bAutoHostWithPIE = false;
+		bAutoHostWithPIE = true;
+		bAutoHostWithGame = true;
 
 		GConfig->GetBool(TEXT("RemoteSession"), TEXT("bAutoHostWithGame"), bAutoHostWithGame, GEngineIni);
 		GConfig->GetBool(TEXT("RemoteSession"), TEXT("bAutoHostWithPIE"), bAutoHostWithPIE, GEngineIni);
@@ -51,19 +54,16 @@ public:
 		GConfig->GetInt(TEXT("RemoteSession"), TEXT("Quality"), Quality, GEngineIni);
 		GConfig->GetInt(TEXT("RemoteSession"), TEXT("Framerate"), Framerate, GEngineIni);
 
-		if (PLATFORM_DESKTOP)
+
+		if (PLATFORM_DESKTOP 
+			&& IsRunningDedicatedServer() == false 
+			&& IsRunningCommandlet() == false)
 		{
-			if (GIsEditor)
-			{
 #if WITH_EDITOR
-				PostPieDelegate = FEditorDelegates::PostPIEStarted.AddRaw(this, &FRemoteSessionModule::OnPIEStarted);
-				EndPieDelegate = FEditorDelegates::EndPIE.AddRaw(this, &FRemoteSessionModule::OnPIEEnded);
+			PostPieDelegate = FEditorDelegates::PostPIEStarted.AddRaw(this, &FRemoteSessionModule::OnPIEStarted);
+			EndPieDelegate = FEditorDelegates::EndPIE.AddRaw(this, &FRemoteSessionModule::OnPIEEnded);
 #endif
-			}
-			else if (bAutoHostWithGame)
-			{
-				InitHost();
-			}
+			GameStartDelegate = FCoreDelegates::OnFEngineLoopInitComplete.AddRaw(this, &FRemoteSessionModule::OnGameStarted);
 		}
 	}
 
@@ -82,42 +82,84 @@ public:
 			FEditorDelegates::EndPIE.Remove(EndPieDelegate);
 		}
 #endif
-	}
 
-	virtual void InitClient(const TCHAR* RemoteAddress) override
-	{
-		if (Client.IsValid())
+		if (GameStartDelegate.IsValid())
 		{
-			Client = nullptr;
+			FCoreDelegates::OnFEngineLoopInitComplete.Remove(GameStartDelegate);
 		}
+	}
 
-		Client = MakeShareable(new FRemoteSessionClient(RemoteAddress));
+	void OnGameStarted()
+	{
+		bool IsHostGame = PLATFORM_DESKTOP
+			&& GIsEditor == false
+			&& IsRunningDedicatedServer() == false
+			&& IsRunningCommandlet() == false;
 
+		if (IsHostGame && bAutoHostWithGame)
+		{
+			InitHost();
+		}
 	}
 
-	virtual bool IsClientConnected() const override
+	void OnPIEStarted(bool bSimulating)
 	{
-		return Client.IsValid() && Client->IsConnected();
+		if (bAutoHostWithPIE)
+		{
+			InitHost();
+		}
 	}
 
-	virtual void StopClient() override
+	void OnPIEEnded(bool bSimulating)
 	{
-		Client = nullptr;
+		// always stop, incase it was started via the console
+		StopHost();
 	}
-
-	virtual TSharedPtr<IRemoteSessionRole> GetClient() const override
+	
+	virtual TSharedPtr<IRemoteSessionRole>	CreateClient(const TCHAR* RemoteAddress) override
 	{
+		// todo - remove this and allow multiple clients (and hosts) to be created
+		if (Client.IsValid())
+		{
+			StopClient(Client);
+		}
+		Client = MakeShareable(new FRemoteSessionClient(RemoteAddress));
 		return Client;
 	}
 
+	virtual void StopClient(TSharedPtr<IRemoteSessionRole> InClient) override
+	{
+		if (InClient.IsValid())
+		{
+			TSharedPtr<FRemoteSessionClient> CastClient = StaticCastSharedPtr<FRemoteSessionClient>(InClient);
+			CastClient->Close();
+			
+			if (CastClient == Client)
+			{
+				Client = nullptr;
+			}
+		}
+	}
+
+
 	virtual void InitHost(const int16 Port = 0) override
 	{
-#if !UE_BUILD_SHIPPING
 		if (Host.IsValid())
 		{
 			Host = nullptr;
 		}
 
+#if UE_BUILD_SHIPPING
+		bool bAllowInShipping = false;
+		GConfig->GetBool(TEXT("RemoteSession"), TEXT("bAllowInShipping"), bAllowInShipping, GEngineIni);
+
+		if (bAllowInShipping == false)
+		{
+			UE_LOG(LogRemoteSession, Log, TEXT("RemoteSession is disabled. Shipping=1"));
+			return;
+		}
+#endif
+
 		TSharedPtr<FRemoteSessionHost> NewHost = MakeShareable(new FRemoteSessionHost(Quality, Framerate));
 
 		int16 SelectedPort = Port ? Port : (int16)DefaultPort;
@@ -131,9 +173,6 @@ public:
 		{
 			UE_LOG(LogRemoteSession, Error, TEXT("Failed to start host listening on port %d"), SelectedPort);
 		}
-#else
-		UE_LOG(LogRemoteSession, Log, TEXT("RemoteSession is disabled. Shipping=1"));
-#endif
 	}
 
 	virtual bool IsHostRunning() const override
@@ -156,20 +195,6 @@ public:
 		return Host;
 	}
 
-	void OnPIEStarted(bool bSimulating)
-	{
-		if (bAutoHostWithPIE)
-		{
-			InitHost();
-		}
-	}
-
-	void OnPIEEnded(bool bSimulating)
-	{
-		// always stop, incase it was started via the console
-		StopHost();
-	}
-
 	virtual TStatId GetStatId() const override
 	{
 		RETURN_QUICK_DECLARE_CYCLE_STAT(FRemoteSession, STATGROUP_Tickables);
@@ -217,7 +242,7 @@ FAutoConsoleCommand GRemoteDisconnectCommand(
 	{
 		if (FRemoteSessionModule* Viewer = FModuleManager::LoadModulePtr<FRemoteSessionModule>("RemoteSession"))
 		{
-			Viewer->StopClient();
+			//Viewer->StopClient();
 			Viewer->StopHost();
 		}
 	})
diff --git a/Engine/Plugins/Experimental/RemoteSession/Source/RemoteSession/Private/RemoteSessionRole.cpp b/Engine/Plugins/Experimental/RemoteSession/Source/RemoteSession/Private/RemoteSessionRole.cpp
--- a/Engine/Plugins/Experimental/RemoteSession/Source/RemoteSession/Private/RemoteSessionRole.cpp
+++ b/Engine/Plugins/Experimental/RemoteSession/Source/RemoteSession/Private/RemoteSessionRole.cpp
@@ -3,7 +3,13 @@
 
 #include "RemoteSessionRole.h"
 #include "RemoteSession.h"
-#include "Channels/RemoteSessionChannel.h"
+#include "BackChannel/Protocol/OSC/BackChannelOSCMessage.h"
+#include "Channels/RemoteSessionInputChannel.h"
+#include "Channels/RemoteSessionXRTrackingChannel.h"
+#include "Channels/RemoteSessionARCameraChannel.h"
+#include "Channels/RemoteSessionFrameBufferChannel.h"
+#include "Async/Async.h"
+#include "GeneralProjectSettings.h"
 
 DEFINE_LOG_CATEGORY(LogRemoteSession);
 
@@ -19,7 +25,13 @@ void FRemoteSessionRole::Close()
 	StopBackgroundThread();
 	OSCConnection = nullptr;
 	Connection = nullptr;
-	Channels.Empty();
+	ClearChannels();
+}
+
+void FRemoteSessionRole::CloseWithError(const FString& Message)
+{
+	ErrorMessage = Message;
+	Close();
 }
 
 void FRemoteSessionRole::Tick(float DeltaTime)
@@ -112,6 +124,139 @@ void FRemoteSessionRole::StopBackgroundThread()
 	}
 }
 
+void FRemoteSessionRole::CreateOSCConnection(TSharedRef<IBackChannelConnection> InConnection)
+{
+	OSCConnection = MakeShareable(new FBackChannelOSCConnection(InConnection));
+	
+	auto Delegate = FBackChannelDispatchDelegate::FDelegate::CreateRaw(this, &FRemoteSessionRole::OnVersionCheck);
+	OSCConnection->AddMessageHandler(TEXT("/Version"),Delegate);
+	
+	OSCConnection->StartReceiveThread();
+	
+	SendVersion();
+}
+
+FString FRemoteSessionRole::GetVersion() const
+{
+	return REMOTE_SESSION_VERSION_STRING;
+	//return FString::Printf(TEXT("%d"), FMath::RandHelper(1000));
+}
+
+void FRemoteSessionRole::SendVersion()
+{
+	// now ask the client to start these channels
+	FBackChannelOSCMessage Msg(TEXT("/Version"));
+	Msg.Write(GetVersion());
+	OSCConnection->SendPacket(Msg);
+}
+
+void FRemoteSessionRole::OnVersionCheck(FBackChannelOSCMessage& Message, FBackChannelOSCDispatch& Dispatch)
+{
+	FString VersionString;
+	
+	Message.Read(VersionString);
+	
+	FString VersionErrorMessage;
+	
+	if (VersionString.Len() == 0)
+	{
+		VersionErrorMessage = TEXT("FRemoteSessionRole: Failed to read version string");
+	}
+	
+	if (VersionString != GetVersion())
+	{
+		VersionErrorMessage = FString::Printf(TEXT("FRemoteSessionRole: Version mismatch. Local=%s, Remote=%s"), *GetVersion(), *VersionString);
+	}
+	
+	if (VersionErrorMessage.Len() > 0)
+	{
+		UE_LOG(LogRemoteSession, Error, TEXT("%s"), *VersionErrorMessage);
+		UE_LOG(LogRemoteSession, Log, TEXT("FRemoteSessionRole: Closing connection due to version mismatch"));
+		CloseWithError(VersionErrorMessage);
+	}
+	else
+	{
+		// Run on the main thread incase our derived states touch anything there (Create Channels does)
+		AsyncTask(ENamedThreads::GameThread, [this]
+		{
+			UE_LOG(LogRemoteSession, Log, TEXT("FRemoteSessionRole: Binding endpoints and creating channels"));
+			OnBindEndpoints();
+			OnCreateChannels();
+		});
+	}
+}
+
+void FRemoteSessionRole::OnBindEndpoints()
+{
+}
+
+void FRemoteSessionRole::OnCreateChannels()
+{
+}
+
+void FRemoteSessionRole::CreateChannel(const FString& InChannelName, ERemoteSessionChannelMode InMode)
+{
+	TSharedPtr<IRemoteSessionChannel> NewChannel;
+	
+	if (InChannelName == FRemoteSessionInputChannel::StaticType())
+	{
+		NewChannel = MakeShareable(new FRemoteSessionInputChannel(InMode, OSCConnection));
+	}
+	else if (InChannelName == FRemoteSessionFrameBufferChannel::StaticType())
+	{
+		NewChannel = MakeShareable(new FRemoteSessionFrameBufferChannel(InMode, OSCConnection));
+	}
+	else if (InChannelName == FRemoteSessionXRTrackingChannel::StaticType())
+	{
+		NewChannel = MakeShareable(new FRemoteSessionXRTrackingChannel(InMode, OSCConnection));
+	}
+	else if (InChannelName == FRemoteSessionARCameraChannel::StaticType())
+	{
+		// Client side sending only works on iOS with Android coming in the future
+		bool IsSupported = (InMode == ERemoteSessionChannelMode::Receive) || PLATFORM_IOS;
+		if (IsSupported)
+		{
+			NewChannel = MakeShareable(new FRemoteSessionARCameraChannel(InMode, OSCConnection));
+		}
+		else
+		{
+			UE_LOG(LogRemoteSession, Warning, TEXT("FRemoteSessionARCameraChannel does not support sending on this platform"));
+		}
+	}
+	
+	if (NewChannel.IsValid())
+	{
+		UE_LOG(LogRemoteSession, Log, TEXT("Created Channel %s with mode %d"), *InChannelName, (int)InMode);
+		Channels.Add(NewChannel);
+	}
+	else
+	{
+		UE_LOG(LogRemoteSession, Error, TEXT("Requested Channel %s was not recognized"), *InChannelName);
+	}
+}
+
+void FRemoteSessionRole::CreateChannels(const TMap<FString, ERemoteSessionChannelMode>& ChannelMap)
+{
+	ClearChannels();
+	
+	for (const auto& KP : ChannelMap)
+	{
+		TSharedPtr<IRemoteSessionChannel> NewChannel;
+		
+		CreateChannel(KP.Key, KP.Value);
+	}
+}
+
+void FRemoteSessionRole::AddChannel(const TSharedPtr<IRemoteSessionChannel>& InChannel)
+{
+	Channels.Add(InChannel);
+}
+
+void FRemoteSessionRole::ClearChannels()
+{
+	Channels.Empty();
+}
+
 TSharedPtr<IRemoteSessionChannel> FRemoteSessionRole::GetChannel(const FString& InType)
 {
 	TSharedPtr<IRemoteSessionChannel> Channel;
diff --git a/Engine/Plugins/Experimental/RemoteSession/Source/RemoteSession/Private/RemoteSessionRole.h b/Engine/Plugins/Experimental/RemoteSession/Source/RemoteSession/Private/RemoteSessionRole.h
--- a/Engine/Plugins/Experimental/RemoteSession/Source/RemoteSession/Private/RemoteSessionRole.h
+++ b/Engine/Plugins/Experimental/RemoteSession/Source/RemoteSession/Private/RemoteSessionRole.h
@@ -4,6 +4,7 @@
 
 #include "RemoteSession/RemoteSessionRole.h"
 #include "BackChannel/Protocol/OSC/BackChannelOSCConnection.h"
+#include "Channels/RemoteSessionChannel.h"
 #include "Tickable.h"
 
 
@@ -15,8 +16,14 @@ public:
 	virtual ~FRemoteSessionRole();
 
 	virtual void Close();
+	
+	virtual void CloseWithError(const FString& Message);
 
 	virtual bool IsConnected() const;
+	
+	virtual bool HasError() const { return ErrorMessage.Len() > 0; }
+	
+	virtual FString GetErrorMessage() const { return ErrorMessage; }
 
 	virtual void Tick( float DeltaTime );
 
@@ -30,6 +37,27 @@ protected:
 	void			StopBackgroundThread();
 
 	uint32			Run();
+	
+	void			CreateOSCConnection(TSharedRef<IBackChannelConnection> InConnection);
+	
+	FString			GetVersion() const;
+	void			SendVersion();
+	void 			OnVersionCheck(FBackChannelOSCMessage& Message, FBackChannelOSCDispatch& Dispatch);
+	
+	virtual void	OnBindEndpoints();
+	virtual void	OnCreateChannels();
+	
+	void 			CreateChannels(const TMap<FString, ERemoteSessionChannelMode>& ChannelMap);
+	void 			CreateChannel(const FString& InChannelList, ERemoteSessionChannelMode InRole);
+	
+	
+	void	AddChannel(const TSharedPtr<IRemoteSessionChannel>& InChannel);
+	void	ClearChannels();
+	
+	FString GetChannelSelectionEndPoint() const
+	{
+		return TEXT("/ChannelSelection");
+	}
 
 protected:
 	
@@ -37,9 +65,11 @@ protected:
 
 	TSharedPtr<FBackChannelOSCConnection, ESPMode::ThreadSafe> OSCConnection;
 
-	TArray<TSharedPtr<IRemoteSessionChannel>> Channels;
+private:
 	
+	FString					ErrorMessage;
+	
+	TArray<TSharedPtr<IRemoteSessionChannel>> Channels;
 	FThreadSafeBool			ThreadExitRequested;
 	FThreadSafeBool			ThreadRunning;
-
-};
\ No newline at end of file
+};
diff --git a/Engine/Plugins/Experimental/RemoteSession/Source/RemoteSession/RemoteSession.Build.cs b/Engine/Plugins/Experimental/RemoteSession/Source/RemoteSession/RemoteSession.Build.cs
--- a/Engine/Plugins/Experimental/RemoteSession/Source/RemoteSession/RemoteSession.Build.cs
+++ b/Engine/Plugins/Experimental/RemoteSession/Source/RemoteSession/RemoteSession.Build.cs
@@ -52,6 +52,7 @@ public class RemoteSession : ModuleRules
 				"ImageWrapper",
 				"MovieSceneCapture",
 				"Sockets",
+				"EngineSettings",
 				"HeadMountedDisplay",
                 "AugmentedReality",
                 // iOS uses the Apple Image Utils plugin for GPU accellerated JPEG compression
diff --git a/Engine/Plugins/Experimental/RemoteSession/Source/RemoteSession/RemoteSession.h b/Engine/Plugins/Experimental/RemoteSession/Source/RemoteSession/RemoteSession.h
--- a/Engine/Plugins/Experimental/RemoteSession/Source/RemoteSession/RemoteSession.h
+++ b/Engine/Plugins/Experimental/RemoteSession/Source/RemoteSession/RemoteSession.h
@@ -6,6 +6,7 @@
 #include "Modules/ModuleManager.h"
 #include "RemoteSessionRole.h"
 
+#define REMOTE_SESSION_VERSION_STRING TEXT("1.0.4")
 
 REMOTESESSION_API DECLARE_LOG_CATEGORY_EXTERN(LogRemoteSession, Log, All);
 
@@ -22,23 +23,26 @@ public:
 
 	enum
 	{
-		kDefaultPort = 2049
+		kDefaultPort = 2049,
 	};
 
 public:
 	/** Client implementation */
+	
+	/** Returns a reference to the client role (if any) */
+	virtual TSharedPtr<IRemoteSessionRole>		CreateClient(const TCHAR* RemoteAddress) = 0;
 
 	/** Initialize a client that will attempt to connect to the provided address */
-	virtual void InitClient(const TCHAR* RemoteAddress) = 0;
+	//virtual void InitClient(const TCHAR* RemoteAddress) = 0;
 
 	/** Returns true/false based on the connection state of the client */
-	virtual bool IsClientConnected() const = 0;
+	//virtual bool IsClientConnected() const = 0;
 
-	/** Stops the client. After this InitClient() must be called if a new connection is desired */
-	virtual void StopClient() = 0;
+	/** Stops the client. After this CreateClient() must be called if a new connection is desired */
+	virtual void StopClient(TSharedPtr<IRemoteSessionRole> InClient) = 0;
 
 	/** Returns a reference to the client role (if any) */
-	virtual TSharedPtr<IRemoteSessionRole>		GetClient() const = 0;
+	//virtual TSharedPtr<IRemoteSessionRole>		GetClient() const = 0;
 
 public:
 	/** Server implementation */
diff --git a/Engine/Plugins/Experimental/RemoteSession/Source/RemoteSession/RemoteSessionRole.h b/Engine/Plugins/Experimental/RemoteSession/Source/RemoteSession/RemoteSessionRole.h
--- a/Engine/Plugins/Experimental/RemoteSession/Source/RemoteSession/RemoteSessionRole.h
+++ b/Engine/Plugins/Experimental/RemoteSession/Source/RemoteSession/RemoteSessionRole.h
@@ -11,12 +11,18 @@ class REMOTESESSION_API IRemoteSessionRole
 public:
 	virtual~IRemoteSessionRole() {}
 
+	virtual bool IsConnected() const = 0;
+	
+	virtual bool HasError() const = 0;
+	
+	virtual FString GetErrorMessage() const = 0;
+
 	virtual TSharedPtr<IRemoteSessionChannel> GetChannel(const FString& Type) = 0;
 
 	template<class T>
-	TSharedPtr<T> GetChannel(const FString& InType)
+	TSharedPtr<T> GetChannel()
 	{
-		TSharedPtr<IRemoteSessionChannel> Channel = GetChannel(InType);
+		TSharedPtr<IRemoteSessionChannel> Channel = GetChannel(T::StaticType());
 
 		if (Channel.IsValid())
 		{
diff --git a/Engine/Source/Editor/LevelEditor/Public/ILevelViewport.h b/Engine/Source/Editor/LevelEditor/Public/ILevelViewport.h
--- a/Engine/Source/Editor/LevelEditor/Public/ILevelViewport.h
+++ b/Engine/Source/Editor/LevelEditor/Public/ILevelViewport.h
@@ -1,4 +1,4 @@
-// Copyright 1998-2018 Epic Games, Inc. All Rights Reserved.
+﻿// Copyright 1998-2018 Epic Games, Inc. All Rights Reserved.
 
 #pragma once
 
@@ -56,6 +56,8 @@ public:
 	 */
 	virtual bool HasPlayInEditorViewport() const = 0; 
 
+	virtual TSharedPtr<FSceneViewport> GetSharedActiveViewport() const  = 0;
+
 	/**
 	 * @return The editor client for this viewport
 	 */
diff --git a/Engine/Source/Editor/LevelEditor/Public/SLevelViewport.h b/Engine/Source/Editor/LevelEditor/Public/SLevelViewport.h
--- a/Engine/Source/Editor/LevelEditor/Public/SLevelViewport.h
+++ b/Engine/Source/Editor/LevelEditor/Public/SLevelViewport.h
@@ -1,4 +1,4 @@
-// Copyright 1998-2018 Epic Games, Inc. All Rights Reserved.
+﻿// Copyright 1998-2018 Epic Games, Inc. All Rights Reserved.
 
 
 #pragma once
@@ -112,6 +112,7 @@ public:
 	virtual void RegisterGameViewportIfPIE() override;
 	virtual bool HasPlayInEditorViewport() const override; 
 	virtual FViewport* GetActiveViewport() override;
+	TSharedPtr<FSceneViewport> GetSharedActiveViewport() const override {return ActiveViewport;};
 	virtual TSharedRef< const SWidget> AsWidget() const override { return AsShared(); }
 	virtual TSharedRef< SWidget> AsWidget() override { return AsShared(); }
 	virtual TWeakPtr< SViewport > GetViewportWidget() override { return ViewportWidget; }
diff --git a/Engine/Source/Runtime/Core/Public/Math/UnitConversion.inl b/Engine/Source/Runtime/Core/Public/Math/UnitConversion.inl
--- a/Engine/Source/Runtime/Core/Public/Math/UnitConversion.inl
+++ b/Engine/Source/Runtime/Core/Public/Math/UnitConversion.inl
@@ -11,7 +11,9 @@ struct FUnitConversion;
 enum class EUnit : uint8;
 enum class EUnitType;
 template<typename NumericType> struct FNumericUnit;
+/// @cond DOXYGEN_WARNINGS
 template<typename OptionalType> struct TOptional;
+/// @endcond
 template<typename ValueType, typename ErrorType> class TValueOrError;
 
 namespace UnitConversion
diff --git a/Engine/Source/Runtime/Engine/Public/CollisionQueryParams.h b/Engine/Source/Runtime/Engine/Public/CollisionQueryParams.h
--- a/Engine/Source/Runtime/Engine/Public/CollisionQueryParams.h
+++ b/Engine/Source/Runtime/Engine/Public/CollisionQueryParams.h
@@ -57,7 +57,7 @@ struct ENGINE_API FCollisionQueryParams
 	/** Whether we want to return the triangle face index for complex static mesh traces */
 	bool bReturnFaceIndex;
 
-	/** Only fill in the PhysMaterial field of  */
+	/** Whether we want to include the physical material in the results. */
 	bool bReturnPhysicalMaterial;
 
 	/** Whether to ignore blocking results. */
diff --git a/Engine/Source/Runtime/UMG/Private/Components/Image.cpp b/Engine/Source/Runtime/UMG/Private/Components/Image.cpp
--- a/Engine/Source/Runtime/UMG/Private/Components/Image.cpp
+++ b/Engine/Source/Runtime/UMG/Private/Components/Image.cpp
@@ -101,7 +101,7 @@ void UImage::SetBrush(const FSlateBrush& InBrush)
 void UImage::SetBrushSize(FVector2D DesiredSize)
 {
 	Brush.ImageSize = DesiredSize;
-
+	
 	if (MyImage.IsValid())
 	{
 		MyImage->SetImage(&Brush);
@@ -109,17 +109,6 @@ void UImage::SetBrushSize(FVector2D DesiredSize)
 	}
 }
 
-void UImage::SetBrushTintColor(FSlateColor TintColor)
-{
-	Brush.TintColor = TintColor;
-
-	if (MyImage.IsValid())
-	{
-		MyImage->SetImage(&Brush);
-		MyImage->Invalidate(EInvalidateWidget::PaintAndVolatility);
-	}
-}
-
 void UImage::SetBrushFromAsset(USlateBrushAsset* Asset)
 {
 	Brush = Asset ? Asset->Brush : FSlateBrush();
@@ -233,7 +222,6 @@ UMaterialInstanceDynamic* UImage::GetDynamicMaterial()
 			if ( MyImage.IsValid() )
 			{
 				MyImage->SetImage(&Brush);
-				MyImage->Invalidate(EInvalidateWidget::LayoutAndVolatility);
 			}
 		}
 		
