Clean up RemoteSession now that force parameter is present
#rnx

diff --git a/Engine/Plugins/Experimental/RemoteSession/Source/RemoteSession/Private/MessageHandler/ProxyMessageHandler.cpp b/Engine/Plugins/Experimental/RemoteSession/Source/RemoteSession/Private/MessageHandler/ProxyMessageHandler.cpp
--- a/Engine/Plugins/Experimental/RemoteSession/Source/RemoteSession/Private/MessageHandler/ProxyMessageHandler.cpp
+++ b/Engine/Plugins/Experimental/RemoteSession/Source/RemoteSession/Private/MessageHandler/ProxyMessageHandler.cpp
@@ -198,37 +198,21 @@ void FProxyMessageHandler::OnEndGesture()
 	}
 }
 
-#if REMOTE_WITH_FORCE_PARAM
 bool FProxyMessageHandler::OnTouchStarted(const TSharedPtr< FGenericWindow >& Window, const FVector2D& Location, float Force, int32 TouchIndex, int32 ControllerId)
-#else
-bool FProxyMessageHandler::OnTouchStarted(const TSharedPtr< FGenericWindow >& Window, const FVector2D& Location, int32 TouchIndex, int32 ControllerId)
-#endif // REMOTE_WITH_FORCE_PARAM
 {
 	if (TargetHandler.IsValid())
 	{
-#if REMOTE_WITH_FORCE_PARAM
 		return TargetHandler->OnTouchStarted(Window, Location, Force, TouchIndex, ControllerId);
-#else
-		return TargetHandler->OnTouchStarted(Window, Location, TouchIndex, ControllerId);
-#endif
 	}
 
 	return false;
 }
 
-#if REMOTE_WITH_FORCE_PARAM
 bool FProxyMessageHandler::OnTouchMoved(const FVector2D& Location, float Force, int32 TouchIndex, int32 ControllerId)
-#else
-bool FProxyMessageHandler::OnTouchMoved(const FVector2D& Location, int32 TouchIndex, int32 ControllerId)
-#endif // REMOTE_WITH_FORCE_PARAM
 {
 	if (TargetHandler.IsValid())
 	{
-#if REMOTE_WITH_FORCE_PARAM
 		return TargetHandler->OnTouchMoved(Location, Force, TouchIndex, ControllerId);
-#else
-		return TargetHandler->OnTouchMoved(Location, TouchIndex, ControllerId);
-#endif
 	}
 
 	return false;
diff --git a/Engine/Plugins/Experimental/RemoteSession/Source/RemoteSession/Private/MessageHandler/ProxyMessageHandler.h b/Engine/Plugins/Experimental/RemoteSession/Source/RemoteSession/Private/MessageHandler/ProxyMessageHandler.h
--- a/Engine/Plugins/Experimental/RemoteSession/Source/RemoteSession/Private/MessageHandler/ProxyMessageHandler.h
+++ b/Engine/Plugins/Experimental/RemoteSession/Source/RemoteSession/Private/MessageHandler/ProxyMessageHandler.h
@@ -6,11 +6,6 @@
 #include "Launch/Resources/Version.h"
 #include "GenericPlatform/GenericApplicationMessageHandler.h"
 
-// For compatibility with a change that has not yet made it to //UE4/Main.
-// Please let andrew.grant know when you run into this as we may want to do some
-// cleanup if it occurs prior to the 4.20 branch.
-#define REMOTE_WITH_FORCE_PARAM 0
-
 class FProxyMessageHandler : public FGenericApplicationMessageHandler
 {
 	
@@ -62,15 +57,9 @@ public:
 
 	virtual void OnEndGesture() override;
 
-#if REMOTE_WITH_FORCE_PARAM
 	virtual bool OnTouchStarted(const TSharedPtr< FGenericWindow >& Window, const FVector2D& Location, float Force, int32 TouchIndex, int32 ControllerId) override;
 
 	virtual bool OnTouchMoved(const FVector2D& Location, float Force, int32 TouchIndex, int32 ControllerId) override;
-#else
-	virtual bool OnTouchStarted(const TSharedPtr< FGenericWindow >& Window, const FVector2D& Location, int32 TouchIndex, int32 ControllerId) override;
-
-	virtual bool OnTouchMoved(const FVector2D& Location, int32 TouchIndex, int32 ControllerId) override;
-#endif // REMOTE_WITH_FORCE_PARAM
 
 	virtual bool OnTouchEnded(const FVector2D& Location, int32 TouchIndex, int32 ControllerId) override;
 
diff --git a/Engine/Plugins/Experimental/RemoteSession/Source/RemoteSession/Private/MessageHandler/RecordingMessageHandler.cpp b/Engine/Plugins/Experimental/RemoteSession/Source/RemoteSession/Private/MessageHandler/RecordingMessageHandler.cpp
--- a/Engine/Plugins/Experimental/RemoteSession/Source/RemoteSession/Private/MessageHandler/RecordingMessageHandler.cpp
+++ b/Engine/Plugins/Experimental/RemoteSession/Source/RemoteSession/Private/MessageHandler/RecordingMessageHandler.cpp
@@ -1,4 +1,4 @@
-﻿// Copyright 1998-2018 Epic Games, Inc. All Rights Reserved.
+// Copyright 1998-2018 Epic Games, Inc. All Rights Reserved.
 
 #include "RecordingMessageHandler.h"
 #include "RemoteSession.h"
@@ -239,20 +239,12 @@ void FRecordingMessageHandler::PlayOnKeyUp(FArchive& Ar)
 	OnKeyUp(Msg.Param1, Msg.Param2, Msg.Param3);
 }
 
-#if REMOTE_WITH_FORCE_PARAM
 bool FRecordingMessageHandler::OnTouchStarted(const TSharedPtr< FGenericWindow >& Window, const FVector2D& Location, float Force, int32 TouchIndex, int32 ControllerId)
-#else
-bool FRecordingMessageHandler::OnTouchStarted(const TSharedPtr< FGenericWindow >& Window, const FVector2D& Location, int32 TouchIndex, int32 ControllerId)
-#endif
 {
 	if (IsRecording())
 	{
 		FVector2D Normalized;
 
-#if !REMOTE_WITH_FORCE_PARAM
-		float Force = 1.0f;
-#endif
-
 		if (ConvertToNormalizedScreenLocation(Location, Normalized))
 		{
 			// note - force is serialized last for backwards compat - force was introduced in 4.20
@@ -268,11 +260,7 @@ bool FRecordingMessageHandler::OnTouchStarted(const TSharedPtr< FGenericWindow >
 		return true;
 	}
 
-#if REMOTE_WITH_FORCE_PARAM
 	return FProxyMessageHandler::OnTouchStarted(Window, Location, Force, TouchIndex, ControllerId);
-#else
-	return FProxyMessageHandler::OnTouchStarted(Window, Location, TouchIndex, ControllerId);
-#endif
 }
 
 void FRecordingMessageHandler::PlayOnTouchStarted(FArchive& Ar)
@@ -287,27 +275,16 @@ void FRecordingMessageHandler::PlayOnTouchStarted(FArchive& Ar)
 		Window = PlaybackWindow.Pin()->GetNativeWindow();
 	}
 
-#if REMOTE_WITH_FORCE_PARAM
 	// note - force is serialized last for backwards compat - force was introduced in 4.20
 	OnTouchStarted(Window, ScreenLocation, Msg.Param4, Msg.Param2, Msg.Param3 );
-#else
-	OnTouchStarted(Window, ScreenLocation, Msg.Param2, Msg.Param3);
-#endif
 }
 
-#if REMOTE_WITH_FORCE_PARAM
 bool FRecordingMessageHandler::OnTouchMoved(const FVector2D& Location, float Force, int32 TouchIndex, int32 ControllerId)
-#else
-bool FRecordingMessageHandler::OnTouchMoved(const FVector2D& Location, int32 TouchIndex, int32 ControllerId)
-#endif
 {
 	if (IsRecording())
 	{
 		FVector2D Normalized;
 
-#if !REMOTE_WITH_FORCE_PARAM
-		float Force = 1.0f;
-#endif
 		if (ConvertToNormalizedScreenLocation(Location, Normalized))
 		{
 			// note - force is serialized last for backwards compat - force was introduced in 4.20
@@ -323,23 +300,15 @@ bool FRecordingMessageHandler::OnTouchMoved(const FVector2D& Location, int32 Tou
 		return true;
 	}
 
-#if REMOTE_WITH_FORCE_PARAM
 	return FProxyMessageHandler::OnTouchMoved(Location, Force, TouchIndex, ControllerId);
-#else
-	return FProxyMessageHandler::OnTouchMoved(Location, TouchIndex, ControllerId);
-#endif
 }
 
 void FRecordingMessageHandler::PlayOnTouchMoved(FArchive& Ar)
 {
 	FourParamMsg<FVector2D, int32, int32, float > Msg(Ar);
 	FVector2D ScreenLocation = ConvertFromNormalizedScreenLocation(Msg.Param1);
-#if REMOTE_WITH_FORCE_PARAM
 	// note - force is serialized last for backwards compat - force was introduced in 4.20
 	OnTouchMoved(ScreenLocation, Msg.Param4, Msg.Param2, Msg.Param3);
-#else
-	OnTouchMoved(ScreenLocation, Msg.Param2, Msg.Param3);
-#endif
 }
 
 bool FRecordingMessageHandler::OnTouchEnded(const FVector2D& Location, int32 TouchIndex, int32 ControllerId)
diff --git a/Engine/Plugins/Experimental/RemoteSession/Source/RemoteSession/Private/MessageHandler/RecordingMessageHandler.h b/Engine/Plugins/Experimental/RemoteSession/Source/RemoteSession/Private/MessageHandler/RecordingMessageHandler.h
--- a/Engine/Plugins/Experimental/RemoteSession/Source/RemoteSession/Private/MessageHandler/RecordingMessageHandler.h
+++ b/Engine/Plugins/Experimental/RemoteSession/Source/RemoteSession/Private/MessageHandler/RecordingMessageHandler.h
@@ -81,13 +81,8 @@ public:
 	virtual bool OnTouchGesture(EGestureEvent GestureType, const FVector2D& Delta, float WheelDelta, bool bIsDirectionInvertedFromDevice) override;
 	virtual void OnEndGesture() override;
 
-#if REMOTE_WITH_FORCE_PARAM
 	virtual bool OnTouchStarted(const TSharedPtr< FGenericWindow >& Window, const FVector2D& Location, float Force, int32 TouchIndex, int32 ControllerId) override;
 	virtual bool OnTouchMoved(const FVector2D& Location, float Force, int32 TouchIndex, int32 ControllerId) override;
-#else
-	virtual bool OnTouchStarted(const TSharedPtr< FGenericWindow >& Window, const FVector2D& Location,  int32 TouchIndex, int32 ControllerId) override;
-	virtual bool OnTouchMoved(const FVector2D& Location, int32 TouchIndex, int32 ControllerId) override;
-#endif // #if REMOTE_WITH_FORCE_PARAM
 	virtual bool OnTouchEnded(const FVector2D& Location, int32 TouchIndex, int32 ControllerId) override;
 	virtual bool OnMotionDetected(const FVector& Tilt, const FVector& RotationRate, const FVector& Gravity, const FVector& Acceleration, int32 ControllerId) override;
 
