Fixed Win32 Build Issue
Fixed RemoteSession crash if a project specified an empty channel list.
Changed channel syntax to be more inline with regular ini syntax
Added AR/XR channels to RemoteSession section of VirtualCameraSample EngineConfig.

diff --git a/Engine/Plugins/Experimental/RemoteSession/Source/RemoteSession/Private/Channels/RemoteSessionARCameraChannel.cpp b/Engine/Plugins/Experimental/RemoteSession/Source/RemoteSession/Private/Channels/RemoteSessionARCameraChannel.cpp
--- a/Engine/Plugins/Experimental/RemoteSession/Source/RemoteSession/Private/Channels/RemoteSessionARCameraChannel.cpp
+++ b/Engine/Plugins/Experimental/RemoteSession/Source/RemoteSession/Private/Channels/RemoteSessionARCameraChannel.cpp
@@ -1,6 +1,7 @@
 // Copyright 1998-2018 Epic Games, Inc. All Rights Reserved.
 
 #include "Channels/RemoteSessionARCameraChannel.h"
+#include "RemoteSession.h"
 #include "Framework/Application/SlateApplication.h"
 #include "BackChannel/Protocol/OSC/BackChannelOSCConnection.h"
 #include "BackChannel/Protocol/OSC/BackChannelOSCMessage.h"
@@ -92,7 +93,7 @@ public:
 	{
 	}
 
-	void SetParameters(FRHICommandList& RHICmdList, const FSceneView View)
+	void SetParameters(FRHICommandList& RHICmdList, const FSceneView& View)
 	{
 		const FVertexShaderRHIParamRef ShaderRHI = GetVertexShader();
 		FMaterialShader::SetViewParameters(RHICmdList, ShaderRHI, View, View.ViewUniformBuffer);
@@ -139,7 +140,7 @@ public:
 		}
 	}
 
-	void SetParameters(FRHICommandList& RHICmdList, const FSceneView View, const FMaterialRenderProxy* Material)
+	void SetParameters(FRHICommandList& RHICmdList, const FSceneView& View, const FMaterialRenderProxy* Material)
 	{
 		const FPixelShaderRHIParamRef ShaderRHI = GetPixelShader();
 		FMaterialShader::SetParameters(RHICmdList, ShaderRHI, Material, *Material->GetMaterial(View.GetFeatureLevel()), View, View.ViewUniformBuffer, ESceneTextureSetupMode::None);
@@ -462,18 +463,25 @@ void FRemoteSessionARCameraChannel::QueueARCameraImage()
 	}
 
 	UARTextureCameraImage* CameraImage = UARBlueprintLibrary::GetCameraImage();
-	if (CameraImage != nullptr && CameraImage->Timestamp > LastQueuedTimestamp)
-	{
-		TSharedPtr<FCompressionTask, ESPMode::ThreadSafe> CompressionTask = MakeShareable(new FCompressionTask());
-		CompressionTask->Width = CameraImage->Size.X;
-		CompressionTask->Height = CameraImage->Size.Y;
-		CompressionTask->AsyncTask = IAppleImageUtilsPlugin::Get().ConvertToJPEG(CameraImage, CVarJPEGQuality.GetValueOnGameThread(), !!CVarJPEGColor.GetValueOnGameThread(), !!CVarJPEGGpu.GetValueOnGameThread());
-		if (CompressionTask->AsyncTask.IsValid())
-		{
-			LastQueuedTimestamp = CameraImage->Timestamp;
-			CompressionQueue.Add(CompressionTask);
-		}
-	}
+	if (CameraImage != nullptr)
+    {
+        if (CameraImage->Timestamp > LastQueuedTimestamp)
+        {
+            TSharedPtr<FCompressionTask, ESPMode::ThreadSafe> CompressionTask = MakeShareable(new FCompressionTask());
+            CompressionTask->Width = CameraImage->Size.X;
+            CompressionTask->Height = CameraImage->Size.Y;
+            CompressionTask->AsyncTask = IAppleImageUtilsPlugin::Get().ConvertToJPEG(CameraImage, CVarJPEGQuality.GetValueOnGameThread(), !!CVarJPEGColor.GetValueOnGameThread(), !!CVarJPEGGpu.GetValueOnGameThread());
+            if (CompressionTask->AsyncTask.IsValid())
+            {
+                LastQueuedTimestamp = CameraImage->Timestamp;
+                CompressionQueue.Add(CompressionTask);
+            }
+        }
+    }
+    else
+    {
+        UE_LOG(LogRemoteSession, Warning, TEXT("No AR Camera Image to send!"));
+    }
 }
 
 void FRemoteSessionARCameraChannel::SendARCameraImage()
diff --git a/Engine/Plugins/Experimental/RemoteSession/Source/RemoteSession/Private/Channels/RemoteSessionXRTrackingChannel.cpp b/Engine/Plugins/Experimental/RemoteSession/Source/RemoteSession/Private/Channels/RemoteSessionXRTrackingChannel.cpp
--- a/Engine/Plugins/Experimental/RemoteSession/Source/RemoteSession/Private/Channels/RemoteSessionXRTrackingChannel.cpp
+++ b/Engine/Plugins/Experimental/RemoteSession/Source/RemoteSession/Private/Channels/RemoteSessionXRTrackingChannel.cpp
@@ -1,6 +1,7 @@
 // Copyright 1998-2018 Epic Games, Inc. All Rights Reserved.
 
 #include "Channels/RemoteSessionXRTrackingChannel.h"
+#include "RemoteSession.h"
 #include "Framework/Application/SlateApplication.h"
 #include "BackChannel/Protocol/OSC/BackChannelOSCConnection.h"
 #include "BackChannel/Protocol/OSC/BackChannelOSCMessage.h"
@@ -83,27 +84,41 @@ void FRemoteSessionXRTrackingChannel::Tick(const float InDeltaTime)
 
 void FRemoteSessionXRTrackingChannel::SendXRTracking()
 {
-	if (Connection.IsValid() && XRSystem.IsValid() && XRSystem->IsTracking(IXRTrackingSystem::HMDDeviceId))
-	{
-		FVector Location;
-		FQuat Orientation;
-		if (XRSystem->GetCurrentPose(IXRTrackingSystem::HMDDeviceId, Orientation, Location))
-		{
-			FRotator Rotation(Orientation);
-
-			TwoParamMsg<FVector, FRotator> MsgParam(Location, Rotation);
-			FBackChannelOSCMessage Msg(MESSAGE_ADDRESS);
-			Msg.Write(MsgParam.AsData());
-
-			Connection->SendPacket(Msg);
-		}
-	}
+	if (Connection.IsValid())
+    {
+        if (XRSystem.IsValid() && XRSystem->IsTracking(IXRTrackingSystem::HMDDeviceId))
+        {
+            FVector Location;
+            FQuat Orientation;
+            if (XRSystem->GetCurrentPose(IXRTrackingSystem::HMDDeviceId, Orientation, Location))
+            {
+                FRotator Rotation(Orientation);
+
+                TwoParamMsg<FVector, FRotator> MsgParam(Location, Rotation);
+                FBackChannelOSCMessage Msg(MESSAGE_ADDRESS);
+                Msg.Write(MsgParam.AsData());
+
+                Connection->SendPacket(Msg);
+                
+                UE_LOG(LogRemoteSession, Verbose, TEXT("Sent Rotation (%.02f,%.02f,%.02f)"), Rotation.Pitch,Rotation.Yaw,Rotation.Roll);
+            }
+            else
+            {
+                 UE_LOG(LogRemoteSession, Warning, TEXT("Failed to get XRPose"));
+            }
+        }
+        else
+        {
+            UE_LOG(LogRemoteSession, Warning, TEXT("XR Tracking not available to send"));
+        }
+    }
 }
 
 void FRemoteSessionXRTrackingChannel::ReceiveXRTracking(FBackChannelOSCMessage& Message, FBackChannelOSCDispatch& Dispatch)
 {
 	if (!ProxyXRSystem.IsValid())
 	{
+        UE_LOG(LogRemoteSession, Warning, TEXT("XRProxy is invalid. Cannot receive pose"));
 		return;
 	}
 
@@ -116,6 +131,8 @@ void FRemoteSessionXRTrackingChannel::ReceiveXRTracking(FBackChannelOSCMessage&
 	{
 		FMemoryReader Ar(*DataCopy);
 		TwoParamMsg<FVector, FRotator> MsgParam(Ar);
+        
+        UE_LOG(LogRemoteSession, Verbose, TEXT("Received Rotation (%.02f,%.02f,%.02f)"), MsgParam.Param2.Pitch,MsgParam.Param2.Yaw,MsgParam.Param2.Roll);
 
 		FTransform NewTransform(MsgParam.Param2, MsgParam.Param1);
         TaskXRSystem->UpdateTrackingToWorldTransform(NewTransform);
diff --git a/Engine/Plugins/Experimental/RemoteSession/Source/RemoteSession/Private/RemoteSessionHost.cpp b/Engine/Plugins/Experimental/RemoteSession/Source/RemoteSession/Private/RemoteSessionHost.cpp
--- a/Engine/Plugins/Experimental/RemoteSession/Source/RemoteSession/Private/RemoteSessionHost.cpp
+++ b/Engine/Plugins/Experimental/RemoteSession/Source/RemoteSession/Private/RemoteSessionHost.cpp
@@ -1,4 +1,4 @@
-﻿// Copyright 1998-2018 Epic Games, Inc. All Rights Reserved.
+// Copyright 1998-2018 Epic Games, Inc. All Rights Reserved.
 
 #include "RemoteSessionHost.h"
 #include "BackChannel/Transport/IBackChannelTransport.h"
@@ -110,19 +110,32 @@ void FRemoteSessionHost::OnCreateChannels()
 	if (DesiredChannels.Num() == 0)
 	{
 		// Default to Input receive and framebuffer send
-		DesiredChannels.Add(FString::Printf(TEXT("%s,receive"), *FRemoteSessionInputChannel::StaticType()));
-		DesiredChannels.Add(FString::Printf(TEXT("%s,send"), *FRemoteSessionFrameBufferChannel::StaticType()));
+		DesiredChannels.Add(FString::Printf(TEXT("(Name=%s,Mode=Read)"), FRemoteSessionInputChannel::StaticType()));
+		DesiredChannels.Add(FString::Printf(TEXT("(Name=%s,Mode=Write)"), FRemoteSessionFrameBufferChannel::StaticType()));
+        UE_LOG(LogRemoteSession, Log, TEXT("No channels specified. Defaulting to Input and Framebuffer."));
 	}
 	
 	TMap<FString, ERemoteSessionChannelMode> ChannelsWithModes;
 	
-	for (const FString& Channel : DesiredChannels)
+	for (FString Channel : DesiredChannels)
 	{
-		FString ChannelName, Mode;
-		
-		if (Channel.Split(TEXT(","), &ChannelName, &Mode))
+        FString ChannelName, Mode;
+        
+        Channel.TrimStartAndEndInline();
+        
+        if (Channel[0] == TEXT('('))
+        {
+            const TCHAR* ChannelArgs = (*Channel)+1;
+            
+            FParse::Value(ChannelArgs, TEXT("Name="), ChannelName);
+            FParse::Value(ChannelArgs, TEXT("Mode="), Mode);
+        }
+        
+		if (ChannelName.Len() && Mode.Len())
 		{
-			ChannelsWithModes.Add(ChannelName, Mode == TEXT("receive") ? ERemoteSessionChannelMode::Receive : ERemoteSessionChannelMode::Send);
+			ChannelsWithModes.Add(ChannelName, Mode == TEXT("Read") ? ERemoteSessionChannelMode::Receive : ERemoteSessionChannelMode::Send);
+            
+            UE_LOG(LogRemoteSession, Log, TEXT("Will request channel %s in mode %s."), *ChannelName, *Mode);
 		}
 		else
 		{
diff --git a/Engine/Plugins/Experimental/RemoteSession/Source/RemoteSession/RemoteSession.h b/Engine/Plugins/Experimental/RemoteSession/Source/RemoteSession/RemoteSession.h
--- a/Engine/Plugins/Experimental/RemoteSession/Source/RemoteSession/RemoteSession.h
+++ b/Engine/Plugins/Experimental/RemoteSession/Source/RemoteSession/RemoteSession.h
@@ -5,8 +5,9 @@
 #include "CoreMinimal.h"
 #include "Modules/ModuleManager.h"
 #include "RemoteSessionRole.h"
+#include "Channels/RemoteSessionChannel.h"
 
-#define REMOTE_SESSION_VERSION_STRING TEXT("1.0.4")
+#define REMOTE_SESSION_VERSION_STRING TEXT("1.0.5")
 
 REMOTESESSION_API DECLARE_LOG_CATEGORY_EXTERN(LogRemoteSession, Log, All);
 
@@ -32,18 +33,10 @@ public:
 	/** Returns a reference to the client role (if any) */
 	virtual TSharedPtr<IRemoteSessionRole>		CreateClient(const TCHAR* RemoteAddress) = 0;
 
-	/** Initialize a client that will attempt to connect to the provided address */
-	//virtual void InitClient(const TCHAR* RemoteAddress) = 0;
-
-	/** Returns true/false based on the connection state of the client */
-	//virtual bool IsClientConnected() const = 0;
-
 	/** Stops the client. After this CreateClient() must be called if a new connection is desired */
 	virtual void StopClient(TSharedPtr<IRemoteSessionRole> InClient) = 0;
 
-	/** Returns a reference to the client role (if any) */
-	//virtual TSharedPtr<IRemoteSessionRole>		GetClient() const = 0;
-
+    
 public:
 	/** Server implementation */
 
@@ -61,5 +54,4 @@ public:
 
 	/** Returns a reference to the server role (if any) */
 	virtual TSharedPtr<IRemoteSessionRole>		GetHost() const = 0;
-
 };
