UE-59066 Intermittent RemoteSessionApp crash

#4.30

Memory leak fixed.

diff --git a/Engine/Plugins/Experimental/RemoteSession/Source/RemoteSession/Private/Channels/RemoteSessionFrameBufferChannel.cpp b/Engine/Plugins/Experimental/RemoteSession/Source/RemoteSession/Private/Channels/RemoteSessionFrameBufferChannel.cpp
--- a/Engine/Plugins/Experimental/RemoteSession/Source/RemoteSession/Private/Channels/RemoteSessionFrameBufferChannel.cpp
+++ b/Engine/Plugins/Experimental/RemoteSession/Source/RemoteSession/Private/Channels/RemoteSessionFrameBufferChannel.cpp
@@ -1,4 +1,4 @@
-// Copyright 1998-2018 Epic Games, Inc. All Rights Reserved.
+﻿// Copyright 1998-2018 Epic Games, Inc. All Rights Reserved.
 
 #include "Channels/RemoteSessionFrameBufferChannel.h"
 #include "RemoteSession.h"
@@ -223,11 +223,17 @@ void FRemoteSessionFrameBufferChannel::Tick(const float InDeltaTime)
 			FUpdateTextureRegion2D* Region = new FUpdateTextureRegion2D(0, 0, 0, 0, QueuedImage->Width, QueuedImage->Height);
 			TArray<uint8>* TextureData = new TArray<uint8>(MoveTemp(QueuedImage->ImageData));
 
-			DecodedTextures[NextImage]->UpdateTextureRegions(0, 1, Region, 4 * QueuedImage->Width, 8, TextureData->GetData(), [this, NextImage, TextureData](auto InTextureData, auto InRegions) {
+			// cleanup functions, gets executed on the render thread after UpdateTextureRegions
+			TFunction<void(uint8*, const FUpdateTextureRegion2D*)> DataCleanupFunc = [this, NextImage, TextureData](uint8* InTextureData, const FUpdateTextureRegion2D* InRegions)
+			{
 				DecodedTextureIndex = NextImage;
-				delete TextureData; // delete array, not underlying data that UpdateTextureRegions passes us
-				delete InRegions;
-			});
+
+				//this is executed later on the render thread, meanwhile TextureData might have changed
+				delete InTextureData; 
+				delete InRegions; 
+			};
+
+			DecodedTextures[NextImage]->UpdateTextureRegions(0, 1, Region, 4 * QueuedImage->Width, 8, TextureData->GetData(), DataCleanupFunc);
 
 			UE_LOG(LogRemoteSession, Verbose, TEXT("GT: Uploaded image %d"),
 				QueuedImage->ImageIndex);
