Fix VirtualCamera plugin no longer working and some RemoteSession tweaks

- Channel selection now done programatically by VirtualCamera plugin
- VirtualCamera no longer creates as ARCameraChannel (wasn't used, and interefered with tracking).
- Fixed missing object references in RemoteSessionApp
- Added some missing devices to UE4Game provisioning file.

#tests Tested RemoteSession in VirtualCameraSample

diff --git a/Engine/Plugins/Experimental/BackChannel/Source/BackChannel/Private/Transport/BackChannelConnection.cpp b/Engine/Plugins/Experimental/BackChannel/Source/BackChannel/Private/Transport/BackChannelConnection.cpp
--- a/Engine/Plugins/Experimental/BackChannel/Source/BackChannel/Private/Transport/BackChannelConnection.cpp
+++ b/Engine/Plugins/Experimental/BackChannel/Source/BackChannel/Private/Transport/BackChannelConnection.cpp
@@ -297,7 +297,7 @@ bool FBackChannelConnection::WaitForConnection(double InTimeout, TFunction<bool(
 			{
 				TSharedRef<FInternetAddr> RemoteAddress = ISocketSubsystem::Get(PLATFORM_SOCKETSUBSYSTEM)->CreateInternetAddr();
 
-				FSocket* ConnectionSocket = Socket->Accept(*RemoteAddress, TEXT("Incoming Connection"));
+				FSocket* ConnectionSocket = Socket->Accept(*RemoteAddress, TEXT("RemoteConnection"));
 
 				if (ConnectionSocket != nullptr)
 				{
diff --git a/Engine/Plugins/Experimental/RemoteSession/Source/RemoteSession/Channels/RemoteSessionChannel.h b/Engine/Plugins/Experimental/RemoteSession/Source/RemoteSession/Channels/RemoteSessionChannel.h
--- a/Engine/Plugins/Experimental/RemoteSession/Source/RemoteSession/Channels/RemoteSessionChannel.h
+++ b/Engine/Plugins/Experimental/RemoteSession/Source/RemoteSession/Channels/RemoteSessionChannel.h
@@ -3,15 +3,12 @@
 #pragma once
 
 #include "CoreMinimal.h"
-
-enum class ERemoteSessionChannelMode
-{
-	Receive,
-	Send
-};
+#include "RemoteSession/RemoteSession.h"
 
 class FBackChannelOSCConnection;
 
+enum class ERemoteSessionChannelMode;
+
 class REMOTESESSION_API IRemoteSessionChannel
 {
 
diff --git a/Engine/Plugins/Experimental/RemoteSession/Source/RemoteSession/Private/Channels/RemoteSessionARCameraChannel.cpp b/Engine/Plugins/Experimental/RemoteSession/Source/RemoteSession/Private/Channels/RemoteSessionARCameraChannel.cpp
--- a/Engine/Plugins/Experimental/RemoteSession/Source/RemoteSession/Private/Channels/RemoteSessionARCameraChannel.cpp
+++ b/Engine/Plugins/Experimental/RemoteSession/Source/RemoteSession/Private/Channels/RemoteSessionARCameraChannel.cpp
@@ -38,6 +38,7 @@
 
 #include "IImageWrapper.h"
 #include "IImageWrapperModule.h"
+#include "GeneralProjectSettings.h"
 
 #define CAMERA_MESSAGE_ADDRESS TEXT("/ARCamera")
 
@@ -381,12 +382,14 @@ FRemoteSessionARCameraChannel::FRemoteSessionARCameraChannel(ERemoteSessionChann
 	
 	static bool OnceTimeARInit = false;
 	
-	if (!OnceTimeARInit && InRole == ERemoteSessionChannelMode::Send)
-	{
-		// Initialize AR if desired - todo, does this need to check device caps?
-		UARSessionConfig* Settings = NewObject<UARSessionConfig>();
-		UARBlueprintLibrary::StartARSession(Settings);
-		
+	if (!OnceTimeARInit && InRole == ERemoteSessionChannelMode::Write)
+	{
+		// Workaround - we don't want to set bSupportAR in our project as it prevents us running on old devices, but this needs
+		// to be true before we try to init ARKit stuff
+		UGeneralProjectSettings* Settings = const_cast<UGeneralProjectSettings*>(GetDefault<UGeneralProjectSettings>());
+		Settings->bSupportAR = true;
+		UARSessionConfig* Config = NewObject<UARSessionConfig>();
+		UARBlueprintLibrary::StartARSession(Config);
 		OnceTimeARInit = true;
 	}
 	
@@ -401,7 +404,7 @@ FRemoteSessionARCameraChannel::FRemoteSessionARCameraChannel(ERemoteSessionChann
 		MaterialInstanceDynamic->SetTextureParameterValue(CameraImageParamName, DefaultTexture);
 	}
 
-	if (Role == ERemoteSessionChannelMode::Receive)
+	if (Role == ERemoteSessionChannelMode::Read)
 	{
 		// Create our image renderer
 		SceneViewExtension = FSceneViewExtensions::NewExtension<FARCameraSceneViewExtension>(*this);
@@ -414,7 +417,7 @@ FRemoteSessionARCameraChannel::FRemoteSessionARCameraChannel(ERemoteSessionChann
 
 FRemoteSessionARCameraChannel::~FRemoteSessionARCameraChannel()
 {
-	if (Role == ERemoteSessionChannelMode::Receive)
+	if (Role == ERemoteSessionChannelMode::Read)
 	{
 		// Remove the callback so it doesn't call back on an invalid this
 		Connection->RemoveMessageHandler(CAMERA_MESSAGE_ADDRESS, MessageCallbackHandle);
@@ -433,13 +436,13 @@ void FRemoteSessionARCameraChannel::Tick(const float InDeltaTime)
 {
 	// Camera capture only works on iOS right now
 #if PLATFORM_IOS
-	if (Role == ERemoteSessionChannelMode::Send)
+	if (Role == ERemoteSessionChannelMode::Write)
 	{
 		QueueARCameraImage();
 		SendARCameraImage();
 	}
 #endif
-	if (Role == ERemoteSessionChannelMode::Receive)
+	if (Role == ERemoteSessionChannelMode::Read)
 	{
 		UpdateRenderingTexture();
 		if (MaterialInstanceDynamic != nullptr)
diff --git a/Engine/Plugins/Experimental/RemoteSession/Source/RemoteSession/Private/Channels/RemoteSessionFrameBufferChannel.cpp b/Engine/Plugins/Experimental/RemoteSession/Source/RemoteSession/Private/Channels/RemoteSessionFrameBufferChannel.cpp
--- a/Engine/Plugins/Experimental/RemoteSession/Source/RemoteSession/Private/Channels/RemoteSessionFrameBufferChannel.cpp
+++ b/Engine/Plugins/Experimental/RemoteSession/Source/RemoteSession/Private/Channels/RemoteSessionFrameBufferChannel.cpp
@@ -48,7 +48,7 @@ FRemoteSessionFrameBufferChannel::FRemoteSessionFrameBufferChannel(ERemoteSessio
 	Role = InRole;
 	ViewportResized = false;
 
-	if (Role == ERemoteSessionChannelMode::Receive)
+	if (Role == ERemoteSessionChannelMode::Read)
 	{
 		auto Delegate = FBackChannelDispatchDelegate::FDelegate::CreateRaw(this, &FRemoteSessionFrameBufferChannel::ReceiveHostImage);
 		MessageCallbackHandle = InConnection->AddMessageHandler(TEXT("/Screen"), Delegate);
@@ -59,7 +59,7 @@ FRemoteSessionFrameBufferChannel::FRemoteSessionFrameBufferChannel(ERemoteSessio
 
 FRemoteSessionFrameBufferChannel::~FRemoteSessionFrameBufferChannel()
 {
-	if (Role == ERemoteSessionChannelMode::Receive)
+	if (Role == ERemoteSessionChannelMode::Read)
 	{
 		TSharedPtr<FBackChannelOSCConnection, ESPMode::ThreadSafe> LocalConnection = Connection.Pin();
 		if (LocalConnection.IsValid())
@@ -187,7 +187,7 @@ void FRemoteSessionFrameBufferChannel::Tick(const float InDeltaTime)
 		}
 	}
 	
-	if (Role == ERemoteSessionChannelMode::Receive)
+	if (Role == ERemoteSessionChannelMode::Read)
 	{
 		SCOPE_CYCLE_COUNTER(STAT_TextureUpdate);
 
diff --git a/Engine/Plugins/Experimental/RemoteSession/Source/RemoteSession/Private/Channels/RemoteSessionInputChannel.cpp b/Engine/Plugins/Experimental/RemoteSession/Source/RemoteSession/Private/Channels/RemoteSessionInputChannel.cpp
--- a/Engine/Plugins/Experimental/RemoteSession/Source/RemoteSession/Private/Channels/RemoteSessionInputChannel.cpp
+++ b/Engine/Plugins/Experimental/RemoteSession/Source/RemoteSession/Private/Channels/RemoteSessionInputChannel.cpp
@@ -17,7 +17,7 @@ FRemoteSessionInputChannel::FRemoteSessionInputChannel(ERemoteSessionChannelMode
 
 	// if sending input replace the default message handler with a recording version, and set us as the
 	// handler for that data 
-	if (Role == ERemoteSessionChannelMode::Send)
+	if (Role == ERemoteSessionChannelMode::Write)
 	{
 		DefaultHandler = FSlateApplication::Get().GetPlatformApplication()->GetMessageHandler();
 
@@ -40,7 +40,7 @@ FRemoteSessionInputChannel::FRemoteSessionInputChannel(ERemoteSessionChannelMode
 
 FRemoteSessionInputChannel::~FRemoteSessionInputChannel()
 {
-	if (Role == ERemoteSessionChannelMode::Receive)
+	if (Role == ERemoteSessionChannelMode::Read)
 	{
 		// Remove the callback so it doesn't call back on an invalid this
 		Connection->RemoveMessageHandler(TEXT("/MessageHandler/"), MessageCallbackHandle);
diff --git a/Engine/Plugins/Experimental/RemoteSession/Source/RemoteSession/Private/Channels/RemoteSessionXRTrackingChannel.cpp b/Engine/Plugins/Experimental/RemoteSession/Source/RemoteSession/Private/Channels/RemoteSessionXRTrackingChannel.cpp
--- a/Engine/Plugins/Experimental/RemoteSession/Source/RemoteSession/Private/Channels/RemoteSessionXRTrackingChannel.cpp
+++ b/Engine/Plugins/Experimental/RemoteSession/Source/RemoteSession/Private/Channels/RemoteSessionXRTrackingChannel.cpp
@@ -5,9 +5,12 @@
 #include "Framework/Application/SlateApplication.h"
 #include "BackChannel/Protocol/OSC/BackChannelOSCConnection.h"
 #include "BackChannel/Protocol/OSC/BackChannelOSCMessage.h"
+#include "ARSessionConfig.h"
+#include "ARBlueprintLibrary.h"
 #include "MessageHandler/Messages.h"
 #include "Engine/Engine.h"
 #include "Async/Async.h"
+#include "GeneralProjectSettings.h"
 
 bool FXRTrackingProxy::EnumerateTrackedDevices(TArray<int32>& OutDevices, EXRTrackedDeviceType Type)
 {
@@ -43,7 +46,7 @@ FRemoteSessionXRTrackingChannel::FRemoteSessionXRTrackingChannel(ERemoteSessionC
 	// If we are sending, we grab the data from GEngine->XRSystem, otherwise we back the current one up for restore later
 	XRSystem = GEngine->XRSystem;
 	
-	if (Role == ERemoteSessionChannelMode::Receive)
+	if (Role == ERemoteSessionChannelMode::Read)
 	{
 		// Make the proxy and set GEngine->XRSystem to it
 		ProxyXRSystem = MakeShared<FXRTrackingProxy, ESPMode::ThreadSafe>();
@@ -53,11 +56,23 @@ FRemoteSessionXRTrackingChannel::FRemoteSessionXRTrackingChannel(ERemoteSessionC
 		MessageCallbackHandle = Connection->AddMessageHandler(MESSAGE_ADDRESS, Delegate);
 		Connection->SetMessageOptions(MESSAGE_ADDRESS, 1);
 	}
+    else
+    {
+        // Initialize AR if desired - todo, does this need to check device caps?
+#if PLATFORM_IOS
+		// Workaround - we don't want to set bSupportAR in our project as it prevents us running on old devices, but this needs
+		// to be true before we try to init ARKit stuff
+		UGeneralProjectSettings* Settings = const_cast<UGeneralProjectSettings*>(GetDefault<UGeneralProjectSettings>());
+		Settings->bSupportAR = true;
+		UARSessionConfig* Config = NewObject<UARSessionConfig>();
+		UARBlueprintLibrary::StartARSession(Config);
+#endif
+    }
 }
 
 FRemoteSessionXRTrackingChannel::~FRemoteSessionXRTrackingChannel()
 {
-	if (Role == ERemoteSessionChannelMode::Receive)
+	if (Role == ERemoteSessionChannelMode::Read)
 	{
 		// Remove the callback so it doesn't call back on an invalid this
 		Connection->RemoveMessageHandler(MESSAGE_ADDRESS, MessageCallbackHandle);
@@ -76,7 +91,7 @@ FRemoteSessionXRTrackingChannel::~FRemoteSessionXRTrackingChannel()
 void FRemoteSessionXRTrackingChannel::Tick(const float InDeltaTime)
 {
 	// Inbound data gets handled as callbacks
-	if (Role == ERemoteSessionChannelMode::Send)
+	if (Role == ERemoteSessionChannelMode::Write)
 	{
 		SendXRTracking();
 	}
diff --git a/Engine/Plugins/Experimental/RemoteSession/Source/RemoteSession/Private/RemoteSessionHost.cpp b/Engine/Plugins/Experimental/RemoteSession/Source/RemoteSession/Private/RemoteSessionHost.cpp
--- a/Engine/Plugins/Experimental/RemoteSession/Source/RemoteSession/Private/RemoteSessionHost.cpp
+++ b/Engine/Plugins/Experimental/RemoteSession/Source/RemoteSession/Private/RemoteSessionHost.cpp
@@ -7,8 +7,6 @@
 #include "BackChannel/Utils/BackChannelThreadedConnection.h"
 #include "BackChannel/Protocol/OSC/BackChannelOSCMessage.h"
 #include "Channels/RemoteSessionInputChannel.h"
-#include "Channels/RemoteSessionXRTrackingChannel.h"
-#include "Channels/RemoteSessionARCameraChannel.h"
 #include "Channels/RemoteSessionFrameBufferChannel.h"
 #include "Engine/GameEngine.h"
 #include "Framework/Application/SlateApplication.h"
@@ -27,11 +25,12 @@ namespace RemoteSessionEd
 };
 
 
-FRemoteSessionHost::FRemoteSessionHost(int32 InQuality, int32 InFramerate)
+FRemoteSessionHost::FRemoteSessionHost(int32 InQuality, int32 InFramerate, const TMap<FString, ERemoteSessionChannelMode>& InSupportedChannels)
 {
 	HostTCPPort = 0;
 	Quality = InQuality;
 	Framerate = InFramerate;
+    SupportedChannels = InSupportedChannels;
 	SavedEditorDragTriggerDistance = FSlateApplication::Get().GetDragTriggerDistance();
 	IsListenerConnected = false;
 }
@@ -103,47 +102,8 @@ void FRemoteSessionHost::OnCreateChannels()
 	
 	ClearChannels();
 	
-	// Query the list of channels from the hosts ini file.
-	TArray<FString> DesiredChannels;
-	GConfig->GetArray(TEXT("RemoteSession"), TEXT("Channels"), DesiredChannels, GEngineIni);
-	
-	if (DesiredChannels.Num() == 0)
-	{
-		// Default to Input receive and framebuffer send
-		DesiredChannels.Add(FString::Printf(TEXT("(Name=%s,Mode=Read)"), FRemoteSessionInputChannel::StaticType()));
-		DesiredChannels.Add(FString::Printf(TEXT("(Name=%s,Mode=Write)"), FRemoteSessionFrameBufferChannel::StaticType()));
-        UE_LOG(LogRemoteSession, Log, TEXT("No channels specified. Defaulting to Input and Framebuffer."));
-	}
-	
-	TMap<FString, ERemoteSessionChannelMode> ChannelsWithModes;
-	
-	for (FString Channel : DesiredChannels)
-	{
-        FString ChannelName, Mode;
-        
-        Channel.TrimStartAndEndInline();
-        
-        if (Channel[0] == TEXT('('))
-        {
-            const TCHAR* ChannelArgs = (*Channel)+1;
-            
-            FParse::Value(ChannelArgs, TEXT("Name="), ChannelName);
-            FParse::Value(ChannelArgs, TEXT("Mode="), Mode);
-        }
-        
-		if (ChannelName.Len() && Mode.Len())
-		{
-			ChannelsWithModes.Add(ChannelName, Mode == TEXT("Read") ? ERemoteSessionChannelMode::Receive : ERemoteSessionChannelMode::Send);
-            
-            UE_LOG(LogRemoteSession, Log, TEXT("Will request channel %s in mode %s."), *ChannelName, *Mode);
-		}
-		else
-		{
-			UE_LOG(LogRemoteSession, Error, TEXT("Unrecognized channel syntax '%s'. Should be ChannelType,r or ChannelType,s"), *Channel);
-		}
-	}
-	
-	CreateChannels(ChannelsWithModes);
+	CreateChannels(SupportedChannels);
+    
 	IsListenerConnected = true;
 	
 	TWeakPtr<SWindow> InputWindow;
@@ -204,9 +164,9 @@ void FRemoteSessionHost::OnCreateChannels()
 	FBackChannelOSCMessage Msg(*GetChannelSelectionEndPoint());
 	
 	// send these across as a name/mode pair
-	for (const auto& KP : ChannelsWithModes)
+	for (const auto& KP : SupportedChannels)
 	{
-		ERemoteSessionChannelMode ClientMode = (KP.Value == ERemoteSessionChannelMode::Send) ? ERemoteSessionChannelMode::Receive : ERemoteSessionChannelMode::Send;
+		ERemoteSessionChannelMode ClientMode = (KP.Value == ERemoteSessionChannelMode::Write) ? ERemoteSessionChannelMode::Read : ERemoteSessionChannelMode::Write;
 		Msg.Write(KP.Key);
 		Msg.Write((int32)ClientMode);
 	}
@@ -229,11 +189,15 @@ void FRemoteSessionHost::Tick(float DeltaTime)
 			StartListening(HostTCPPort);
 			IsListenerConnected = false;
 		}
-		Listener->WaitForConnection(0, [this](TSharedRef<IBackChannelConnection> InConnection) {
-			Close();
-			CreateOSCConnection(InConnection);
-			return true;
-		});
+        
+        if (Listener.IsValid())
+        {
+            Listener->WaitForConnection(0, [this](TSharedRef<IBackChannelConnection> InConnection) {
+                Close();
+                CreateOSCConnection(InConnection);
+                return true;
+            });
+        }
 	}
 	
 	FRemoteSessionRole::Tick(DeltaTime);
diff --git a/Engine/Plugins/Experimental/RemoteSession/Source/RemoteSession/Private/RemoteSessionHost.h b/Engine/Plugins/Experimental/RemoteSession/Source/RemoteSession/Private/RemoteSessionHost.h
--- a/Engine/Plugins/Experimental/RemoteSession/Source/RemoteSession/Private/RemoteSessionHost.h
+++ b/Engine/Plugins/Experimental/RemoteSession/Source/RemoteSession/Private/RemoteSessionHost.h
@@ -1,4 +1,4 @@
-﻿// Copyright 1998-2018 Epic Games, Inc. All Rights Reserved.
+// Copyright 1998-2018 Epic Games, Inc. All Rights Reserved.
 
 #pragma once
 
@@ -14,7 +14,7 @@ class FRemoteSessionHost : public FRemoteSessionRole, public TSharedFromThis<FRe
 {
 public:
 
-	FRemoteSessionHost(int32 InQuality, int32 InFramerate);
+    FRemoteSessionHost(int32 InQuality, int32 InFramerate, const TMap<FString, ERemoteSessionChannelMode>& InSupportedChannels);
 	~FRemoteSessionHost();
 
 	virtual void Close() override;
@@ -38,6 +38,8 @@ protected:
 
 	int32		Quality;
 	int32		Framerate;
+    
+    TMap<FString, ERemoteSessionChannelMode> SupportedChannels;
 
 	/** Saved information about the editor and viewport we possessed, so we can restore it after exiting VR mode */
 	float SavedEditorDragTriggerDistance;
@@ -47,4 +49,4 @@ protected:
 
 	/** True if the host TCP socket is connected*/
 	bool IsListenerConnected;
-};
\ No newline at end of file
+};
diff --git a/Engine/Plugins/Experimental/RemoteSession/Source/RemoteSession/Private/RemoteSessionModule.cpp b/Engine/Plugins/Experimental/RemoteSession/Source/RemoteSession/Private/RemoteSessionModule.cpp
--- a/Engine/Plugins/Experimental/RemoteSession/Source/RemoteSession/Private/RemoteSessionModule.cpp
+++ b/Engine/Plugins/Experimental/RemoteSession/Source/RemoteSession/Private/RemoteSessionModule.cpp
@@ -5,8 +5,11 @@
 #include "BackChannel/Transport/IBackChannelTransport.h"
 #include "RemoteSessionHost.h"
 #include "RemoteSessionClient.h"
+#include "Channels/RemoteSessionInputChannel.h"
+#include "Channels/RemoteSessionFrameBufferChannel.h"
 #include "CoreGlobals.h"
 #include "Misc/ConfigCacheIni.h"
+#include "Misc/CommandLine.h"
 
 #if WITH_EDITOR
 	#include "Editor.h"
@@ -29,6 +32,8 @@ protected:
 	bool                                bAutoHostWithPIE;
 	bool                                bAutoHostWithGame;
     
+    TMap<FString, ERemoteSessionChannelMode> SupportedChannels;
+    
 	FDelegateHandle PostPieDelegate;
 	FDelegateHandle EndPieDelegate;
 	FDelegateHandle GameStartDelegate;
@@ -49,12 +54,7 @@ public:
 		bAutoHostWithPIE = true;
 		bAutoHostWithGame = true;
 
-		GConfig->GetBool(TEXT("RemoteSession"), TEXT("bAutoHostWithGame"), bAutoHostWithGame, GEngineIni);
-		GConfig->GetBool(TEXT("RemoteSession"), TEXT("bAutoHostWithPIE"), bAutoHostWithPIE, GEngineIni);
-		GConfig->GetInt(TEXT("RemoteSession"), TEXT("HostPort"), DefaultPort, GEngineIni);
-		GConfig->GetInt(TEXT("RemoteSession"), TEXT("Quality"), Quality, GEngineIni);
-		GConfig->GetInt(TEXT("RemoteSession"), TEXT("Framerate"), Framerate, GEngineIni);
-
+        ReadIniSettings();
 
 		if (PLATFORM_DESKTOP 
 			&& IsRunningDedicatedServer() == false 
@@ -67,6 +67,69 @@ public:
 			GameStartDelegate = FCoreDelegates::OnFEngineLoopInitComplete.AddRaw(this, &FRemoteSessionModule::OnGameStarted);
 		}
 	}
+    
+    void ReadIniSettings()
+    {
+        GConfig->GetBool(TEXT("RemoteSession"), TEXT("bAutoHostWithGame"), bAutoHostWithGame, GEngineIni);
+        GConfig->GetBool(TEXT("RemoteSession"), TEXT("bAutoHostWithPIE"), bAutoHostWithPIE, GEngineIni);
+        GConfig->GetInt(TEXT("RemoteSession"), TEXT("HostPort"), DefaultPort, GEngineIni);
+        GConfig->GetInt(TEXT("RemoteSession"), TEXT("Quality"), Quality, GEngineIni);
+        GConfig->GetInt(TEXT("RemoteSession"), TEXT("Framerate"), Framerate, GEngineIni);
+        
+        // Query the list of channels from the hosts ini file.
+        TArray<FString> IniSupportedChannels;
+        GConfig->GetArray(TEXT("RemoteSession"), TEXT("Channels"), IniSupportedChannels, GEngineIni);
+        
+        if (IniSupportedChannels.Num() == 0)
+        {
+            // Default to Input receive and framebuffer send
+            IniSupportedChannels.Add(FString::Printf(TEXT("(Name=%s,Mode=Read)"), FRemoteSessionInputChannel::StaticType()));
+            IniSupportedChannels.Add(FString::Printf(TEXT("(Name=%s,Mode=Write)"), FRemoteSessionFrameBufferChannel::StaticType()));
+            UE_LOG(LogRemoteSession, Log, TEXT("No channels specified. Defaulting to Input and Framebuffer."));
+        }
+        
+        FParse::Value(FCommandLine::Get(), TEXT("remote.port="), DefaultPort);
+        
+        SupportedChannels.Empty();
+        
+        for (FString Channel : IniSupportedChannels)
+        {
+            FString ChannelName, Mode;
+            
+            Channel.TrimStartAndEndInline();
+            
+            if (Channel[0] == TEXT('('))
+            {
+                const TCHAR* ChannelArgs = (*Channel)+1;
+                
+                FParse::Value(ChannelArgs, TEXT("Name="), ChannelName);
+                FParse::Value(ChannelArgs, TEXT("Mode="), Mode);
+            }
+            
+            if (ChannelName.Len() && Mode.Len())
+            {
+                SupportedChannels.Add(ChannelName, Mode == TEXT("Read") ? ERemoteSessionChannelMode::Read : ERemoteSessionChannelMode::Write);
+                
+                UE_LOG(LogRemoteSession, Log, TEXT("Will request channel %s in mode %s."), *ChannelName, *Mode);
+            }
+            else
+            {
+                UE_LOG(LogRemoteSession, Error, TEXT("Unrecognized channel syntax '%s'. Should be ChannelType,r or ChannelType,s"), *Channel);
+            }
+        }
+    }
+    
+    /** Programatically sets the desired channels. Defaults are Input=Receive and Framebuffer=Send. Unioned with values from ini file */
+    virtual void SetSupportedChannels(TMap<FString,ERemoteSessionChannelMode>& InSupportedChannels) override
+    {
+        for (const auto& KP : InSupportedChannels)
+        {
+            if (SupportedChannels.Contains(KP.Key) == false)
+            {
+                SupportedChannels.Add(KP.Key, KP.Value);
+            }
+        }
+    }
 
 	void ShutdownModule()
 	{
@@ -89,6 +152,9 @@ public:
 			FCoreDelegates::OnFEngineLoopInitComplete.Remove(GameStartDelegate);
 		}
 	}
+    
+    
+    
 
 	void OnGameStarted()
 	{
@@ -161,7 +227,7 @@ public:
 		}
 #endif
 
-		TSharedPtr<FRemoteSessionHost> NewHost = MakeShareable(new FRemoteSessionHost(Quality, Framerate));
+		TSharedPtr<FRemoteSessionHost> NewHost = MakeShareable(new FRemoteSessionHost(Quality, Framerate, SupportedChannels));
 
 		int16 SelectedPort = Port ? Port : (int16)DefaultPort;
 
diff --git a/Engine/Plugins/Experimental/RemoteSession/Source/RemoteSession/Private/RemoteSessionRole.cpp b/Engine/Plugins/Experimental/RemoteSession/Source/RemoteSession/Private/RemoteSessionRole.cpp
--- a/Engine/Plugins/Experimental/RemoteSession/Source/RemoteSession/Private/RemoteSessionRole.cpp
+++ b/Engine/Plugins/Experimental/RemoteSession/Source/RemoteSession/Private/RemoteSessionRole.cpp
@@ -83,7 +83,9 @@ void FRemoteSessionRole::StartBackgroundThread()
 
 bool FRemoteSessionRole::IsConnected() const
 {
-	return OSCConnection.IsValid() && OSCConnection->IsConnected();
+    // just check this is valid, when it's actually disconnected we do some error
+    // handling and clean this up
+    return OSCConnection.IsValid();
 }
 
 uint32 FRemoteSessionRole::Run()
@@ -213,7 +215,7 @@ void FRemoteSessionRole::CreateChannel(const FString& InChannelName, ERemoteSess
 	else if (InChannelName == FRemoteSessionARCameraChannel::StaticType())
 	{
 		// Client side sending only works on iOS with Android coming in the future
-		bool IsSupported = (InMode == ERemoteSessionChannelMode::Receive) || PLATFORM_IOS;
+		bool IsSupported = (InMode == ERemoteSessionChannelMode::Read) || PLATFORM_IOS;
 		if (IsSupported)
 		{
 			NewChannel = MakeShareable(new FRemoteSessionARCameraChannel(InMode, OSCConnection));
diff --git a/Engine/Plugins/Experimental/RemoteSession/Source/RemoteSession/Private/RemoteSessionRole.h b/Engine/Plugins/Experimental/RemoteSession/Source/RemoteSession/Private/RemoteSessionRole.h
--- a/Engine/Plugins/Experimental/RemoteSession/Source/RemoteSession/Private/RemoteSessionRole.h
+++ b/Engine/Plugins/Experimental/RemoteSession/Source/RemoteSession/Private/RemoteSessionRole.h
@@ -4,10 +4,11 @@
 
 #include "RemoteSession/RemoteSessionRole.h"
 #include "BackChannel/Protocol/OSC/BackChannelOSCConnection.h"
-#include "Channels/RemoteSessionChannel.h"
-#include "Tickable.h"
 
+#include "Tickable.h"
 
+class FBackChannelOSCConnection;
+enum class ERemoteSessionChannelMode;
 
 class FRemoteSessionRole : public IRemoteSessionRole, FRunnable
 {
diff --git a/Engine/Plugins/Experimental/RemoteSession/Source/RemoteSession/RemoteSession.h b/Engine/Plugins/Experimental/RemoteSession/Source/RemoteSession/RemoteSession.h
--- a/Engine/Plugins/Experimental/RemoteSession/Source/RemoteSession/RemoteSession.h
+++ b/Engine/Plugins/Experimental/RemoteSession/Source/RemoteSession/RemoteSession.h
@@ -11,6 +11,12 @@
 
 REMOTESESSION_API DECLARE_LOG_CATEGORY_EXTERN(LogRemoteSession, Log, All);
 
+enum class ERemoteSessionChannelMode
+{
+    Read,
+    Write
+};
+
 class REMOTESESSION_API IRemoteSessionModule : public IModuleInterface
 {
 public:
@@ -52,6 +58,9 @@ public:
 	/** Stops the server, after this InitHost() must be called if a new connection is desired */
 	virtual void StopHost() = 0;
 
+    /** Programatically sets the desired channels. Defaults are Input=Receive and Framebuffer=Send. Unioned with values from ini file */
+    virtual void SetSupportedChannels(TMap<FString,ERemoteSessionChannelMode>& SupportedChannels) = 0;
+
 	/** Returns a reference to the server role (if any) */
 	virtual TSharedPtr<IRemoteSessionRole>		GetHost() const = 0;
 };
diff --git a/Engine/Plugins/Experimental/VirtualCamera/Source/VirtualCamera/Private/VirtualCameraPlayerControllerBase.cpp b/Engine/Plugins/Experimental/VirtualCamera/Source/VirtualCamera/Private/VirtualCameraPlayerControllerBase.cpp
--- a/Engine/Plugins/Experimental/VirtualCamera/Source/VirtualCamera/Private/VirtualCameraPlayerControllerBase.cpp
+++ b/Engine/Plugins/Experimental/VirtualCamera/Source/VirtualCamera/Private/VirtualCameraPlayerControllerBase.cpp
@@ -5,6 +5,9 @@
 #include "Features/IModularFeatures.h"
 #include "IXRTrackingSystem.h"
 #include "RemoteSession/RemoteSession.h"
+#include "RemoteSession/Channels/RemoteSessionInputChannel.h"
+#include "RemoteSession/Channels/RemoteSessionFrameBufferChannel.h"
+#include "RemoteSession/Channels/RemoteSessionXRTrackingChannel.h"
 #include "SteamVRFunctionLibrary.h"
 
 AVirtualCameraPlayerControllerBase::AVirtualCameraPlayerControllerBase(const FObjectInitializer& ObjectInitializer)
@@ -42,9 +45,15 @@ void AVirtualCameraPlayerControllerBase::BeginPlay()
 		LevelSequencePlaybackController->OnStop.Add(OnStop);
 	}
 
-	if (IRemoteSessionModule* Viewer = FModuleManager::LoadModulePtr<IRemoteSessionModule>("RemoteSession"))
+	if (IRemoteSessionModule* RemoteSession = FModuleManager::LoadModulePtr<IRemoteSessionModule>("RemoteSession"))
 	{
-		Viewer->InitHost();
+        TMap<FString, ERemoteSessionChannelMode> RequiredChannels;
+        RequiredChannels.Add(FRemoteSessionFrameBufferChannel::StaticType(),ERemoteSessionChannelMode::Write);
+        RequiredChannels.Add(FRemoteSessionInputChannel::StaticType(),ERemoteSessionChannelMode::Read);
+        RequiredChannels.Add(FRemoteSessionXRTrackingChannel::StaticType(),ERemoteSessionChannelMode::Read);
+        
+        RemoteSession->SetSupportedChannels(RequiredChannels);
+		RemoteSession->InitHost();
 	}
 	
 	if (UVirtualCameraCineCameraComponent* CineCamera = GetVirtualCameraCineCameraComponent())
