Additional fix for UE-59478

diff --git a/Engine/Plugins/Experimental/RemoteSession/Source/RemoteSession/Private/Channels/RemoteSessionARCameraChannel.cpp b/Engine/Plugins/Experimental/RemoteSession/Source/RemoteSession/Private/Channels/RemoteSessionARCameraChannel.cpp
--- a/Engine/Plugins/Experimental/RemoteSession/Source/RemoteSession/Private/Channels/RemoteSessionARCameraChannel.cpp
+++ b/Engine/Plugins/Experimental/RemoteSession/Source/RemoteSession/Private/Channels/RemoteSessionARCameraChannel.cpp
@@ -38,7 +38,6 @@
 
 #include "IImageWrapper.h"
 #include "IImageWrapperModule.h"
-#include "GeneralProjectSettings.h"
 
 #define CAMERA_MESSAGE_ADDRESS TEXT("/ARCamera")
 
@@ -384,10 +383,6 @@ FRemoteSessionARCameraChannel::FRemoteSessionARCameraChannel(ERemoteSessionChann
 	
 	if (!OnceTimeARInit && InRole == ERemoteSessionChannelMode::Write)
 	{
-		// Workaround - we don't want to set bSupportAR in our project as it prevents us running on old devices, but this needs
-		// to be true before we try to init ARKit stuff
-		UGeneralProjectSettings* Settings = const_cast<UGeneralProjectSettings*>(GetDefault<UGeneralProjectSettings>());
-		Settings->bSupportAR = true;
 		UARSessionConfig* Config = NewObject<UARSessionConfig>();
 		UARBlueprintLibrary::StartARSession(Config);
 		OnceTimeARInit = true;
diff --git a/Engine/Plugins/Experimental/RemoteSession/Source/RemoteSession/Private/Channels/RemoteSessionXRTrackingChannel.cpp b/Engine/Plugins/Experimental/RemoteSession/Source/RemoteSession/Private/Channels/RemoteSessionXRTrackingChannel.cpp
--- a/Engine/Plugins/Experimental/RemoteSession/Source/RemoteSession/Private/Channels/RemoteSessionXRTrackingChannel.cpp
+++ b/Engine/Plugins/Experimental/RemoteSession/Source/RemoteSession/Private/Channels/RemoteSessionXRTrackingChannel.cpp
@@ -10,7 +10,6 @@
 #include "MessageHandler/Messages.h"
 #include "Engine/Engine.h"
 #include "Async/Async.h"
-#include "GeneralProjectSettings.h"
 
 bool FXRTrackingProxy::EnumerateTrackedDevices(TArray<int32>& OutDevices, EXRTrackedDeviceType Type)
 {
@@ -60,10 +59,6 @@ FRemoteSessionXRTrackingChannel::FRemoteSessionXRTrackingChannel(ERemoteSessionC
     {
         // Initialize AR if desired - todo, does this need to check device caps?
 #if PLATFORM_IOS
-		// Workaround - we don't want to set bSupportAR in our project as it prevents us running on old devices, but this needs
-		// to be true before we try to init ARKit stuff
-		UGeneralProjectSettings* Settings = const_cast<UGeneralProjectSettings*>(GetDefault<UGeneralProjectSettings>());
-		Settings->bSupportAR = true;
 		UARSessionConfig* Config = NewObject<UARSessionConfig>();
 		UARBlueprintLibrary::StartARSession(Config);
 #endif
