Merging the RemoteSessionApp changes for bSupportAR removal using Dev-VR_to_Release-4.20

diff --git a/Engine/Plugins/Experimental/RemoteSession/Source/RemoteSession/Private/Channels/RemoteSessionARCameraChannel.cpp b/Engine/Plugins/Experimental/RemoteSession/Source/RemoteSession/Private/Channels/RemoteSessionARCameraChannel.cpp
--- a/Engine/Plugins/Experimental/RemoteSession/Source/RemoteSession/Private/Channels/RemoteSessionARCameraChannel.cpp
+++ b/Engine/Plugins/Experimental/RemoteSession/Source/RemoteSession/Private/Channels/RemoteSessionARCameraChannel.cpp
@@ -378,16 +378,6 @@ FRemoteSessionARCameraChannel::FRemoteSessionARCameraChannel(ERemoteSessionChann
 	, Connection(InConnection)
 	, Role(InRole)
 {
-	
-	static bool OnceTimeARInit = false;
-	
-	if (!OnceTimeARInit && InRole == ERemoteSessionChannelMode::Write)
-	{
-		UARSessionConfig* Config = NewObject<UARSessionConfig>();
-		UARBlueprintLibrary::StartARSession(Config);
-		OnceTimeARInit = true;
-	}
-	
 	RenderingTextures[0] = nullptr;
 	RenderingTextures[1] = nullptr;
 	PPMaterial = LoadObject<UMaterialInterface>(nullptr, TEXT("/RemoteSession/ARCameraPostProcess.ARCameraPostProcess"));
diff --git a/Engine/Plugins/Experimental/RemoteSession/Source/RemoteSession/Private/Channels/RemoteSessionXRTrackingChannel.cpp b/Engine/Plugins/Experimental/RemoteSession/Source/RemoteSession/Private/Channels/RemoteSessionXRTrackingChannel.cpp
--- a/Engine/Plugins/Experimental/RemoteSession/Source/RemoteSession/Private/Channels/RemoteSessionXRTrackingChannel.cpp
+++ b/Engine/Plugins/Experimental/RemoteSession/Source/RemoteSession/Private/Channels/RemoteSessionXRTrackingChannel.cpp
@@ -55,14 +55,6 @@ FRemoteSessionXRTrackingChannel::FRemoteSessionXRTrackingChannel(ERemoteSessionC
 		MessageCallbackHandle = Connection->AddMessageHandler(MESSAGE_ADDRESS, Delegate);
 		Connection->SetMessageOptions(MESSAGE_ADDRESS, 1);
 	}
-    else
-    {
-        // Initialize AR if desired - todo, does this need to check device caps?
-#if PLATFORM_IOS
-		UARSessionConfig* Config = NewObject<UARSessionConfig>();
-		UARBlueprintLibrary::StartARSession(Config);
-#endif
-    }
 }
 
 FRemoteSessionXRTrackingChannel::~FRemoteSessionXRTrackingChannel()
diff --git a/Engine/Plugins/Experimental/RemoteSession/Source/RemoteSession/Private/RemoteSessionRole.cpp b/Engine/Plugins/Experimental/RemoteSession/Source/RemoteSession/Private/RemoteSessionRole.cpp
--- a/Engine/Plugins/Experimental/RemoteSession/Source/RemoteSession/Private/RemoteSessionRole.cpp
+++ b/Engine/Plugins/Experimental/RemoteSession/Source/RemoteSession/Private/RemoteSessionRole.cpp
@@ -208,21 +208,37 @@ void FRemoteSessionRole::CreateChannel(const FString& InChannelName, ERemoteSess
 	{
 		NewChannel = MakeShareable(new FRemoteSessionFrameBufferChannel(InMode, OSCConnection));
 	}
-	else if (InChannelName == FRemoteSessionXRTrackingChannel::StaticType())
+    else if (InChannelName == FRemoteSessionXRTrackingChannel::StaticType())
 	{
-		NewChannel = MakeShareable(new FRemoteSessionXRTrackingChannel(InMode, OSCConnection));
+		// Skip if there isn't an AR session running
+		if (UARBlueprintLibrary::GetSessionStatus().Status == EARSessionStatus::Running))
+		{
+			NewChannel = MakeShareable(new FRemoteSessionXRTrackingChannel(InMode, OSCConnection));
+		}
+		else
+		{
+			UE_LOG(LogRemoteSession, Warning, TEXT("No AR session in progress so FRemoteSessionXRTrackingChannel is disabled"));
+		}
 	}
 	else if (InChannelName == FRemoteSessionARCameraChannel::StaticType())
 	{
-		// Client side sending only works on iOS with Android coming in the future
-		bool IsSupported = (InMode == ERemoteSessionChannelMode::Read) || PLATFORM_IOS;
-		if (IsSupported)
+		// Skip if there isn't an AR session running
+		if (UARBlueprintLibrary::GetSessionStatus().Status == EARSessionStatus::Running))
 		{
-			NewChannel = MakeShareable(new FRemoteSessionARCameraChannel(InMode, OSCConnection));
+			// Client side sending only works on iOS with Android coming in the future
+			bool IsSupported = (InMode == ERemoteSessionChannelMode::Read) || PLATFORM_IOS;
+			if (IsSupported)
+			{
+				NewChannel = MakeShareable(new FRemoteSessionARCameraChannel(InMode, OSCConnection));
+			}
+			else
+			{
+				UE_LOG(LogRemoteSession, Warning, TEXT("FRemoteSessionARCameraChannel does not support sending on this platform"));
+			}
 		}
 		else
 		{
-			UE_LOG(LogRemoteSession, Warning, TEXT("FRemoteSessionARCameraChannel does not support sending on this platform"));
+			UE_LOG(LogRemoteSession, Warning, TEXT("No AR session in progress so FRemoteSessionARCameraChannel is disabled"));
 		}
 	}
 	
