Merging using Dev-VR_to_Release-4.20

diff --git a/Engine/Plugins/Experimental/RemoteSession/Source/RemoteSession/Private/RemoteSessionRole.cpp b/Engine/Plugins/Experimental/RemoteSession/Source/RemoteSession/Private/RemoteSessionRole.cpp
--- a/Engine/Plugins/Experimental/RemoteSession/Source/RemoteSession/Private/RemoteSessionRole.cpp
+++ b/Engine/Plugins/Experimental/RemoteSession/Source/RemoteSession/Private/RemoteSessionRole.cpp
@@ -211,7 +211,7 @@ void FRemoteSessionRole::CreateChannel(const FString& InChannelName, ERemoteSess
     else if (InChannelName == FRemoteSessionXRTrackingChannel::StaticType())
 	{
 		// Skip if there isn't an AR session running
-		if (UARBlueprintLibrary::GetSessionStatus().Status == EARSessionStatus::Running))
+		if (UARBlueprintLibrary::GetARSessionStatus().Status == EARSessionStatus::Running)
 		{
 			NewChannel = MakeShareable(new FRemoteSessionXRTrackingChannel(InMode, OSCConnection));
 		}
@@ -223,7 +223,7 @@ void FRemoteSessionRole::CreateChannel(const FString& InChannelName, ERemoteSess
 	else if (InChannelName == FRemoteSessionARCameraChannel::StaticType())
 	{
 		// Skip if there isn't an AR session running
-		if (UARBlueprintLibrary::GetSessionStatus().Status == EARSessionStatus::Running))
+		if (UARBlueprintLibrary::GetARSessionStatus().Status == EARSessionStatus::Running)
 		{
 			// Client side sending only works on iOS with Android coming in the future
 			bool IsSupported = (InMode == ERemoteSessionChannelMode::Read) || PLATFORM_IOS;
