Merging using Dev-VR_to_Release-4.20

diff --git a/Engine/Plugins/Experimental/RemoteSession/Source/RemoteSession/Private/RemoteSessionRole.cpp b/Engine/Plugins/Experimental/RemoteSession/Source/RemoteSession/Private/RemoteSessionRole.cpp
--- a/Engine/Plugins/Experimental/RemoteSession/Source/RemoteSession/Private/RemoteSessionRole.cpp
+++ b/Engine/Plugins/Experimental/RemoteSession/Source/RemoteSession/Private/RemoteSessionRole.cpp
@@ -10,6 +10,7 @@
 #include "Channels/RemoteSessionFrameBufferChannel.h"
 #include "Async/Async.h"
 #include "GeneralProjectSettings.h"
+#include "ARBlueprintLibrary.h"
 
 DEFINE_LOG_CATEGORY(LogRemoteSession);
 
