Merging host fix for bSupportAR using Dev-VR_to_Release-4.20

diff --git a/Engine/Plugins/Experimental/RemoteSession/Source/RemoteSession/Private/RemoteSessionRole.cpp b/Engine/Plugins/Experimental/RemoteSession/Source/RemoteSession/Private/RemoteSessionRole.cpp
--- a/Engine/Plugins/Experimental/RemoteSession/Source/RemoteSession/Private/RemoteSessionRole.cpp
+++ b/Engine/Plugins/Experimental/RemoteSession/Source/RemoteSession/Private/RemoteSessionRole.cpp
@@ -212,7 +212,9 @@ void FRemoteSessionRole::CreateChannel(const FString& InChannelName, ERemoteSess
     else if (InChannelName == FRemoteSessionXRTrackingChannel::StaticType())
 	{
 		// Skip if there isn't an AR session running
-		if (UARBlueprintLibrary::GetARSessionStatus().Status == EARSessionStatus::Running)
+		if (UARBlueprintLibrary::GetARSessionStatus().Status == EARSessionStatus::Running ||
+			// Add our listener if we are the host
+			InMode == ERemoteSessionChannelMode::Read)
 		{
 			NewChannel = MakeShareable(new FRemoteSessionXRTrackingChannel(InMode, OSCConnection));
 		}
@@ -224,7 +226,9 @@ void FRemoteSessionRole::CreateChannel(const FString& InChannelName, ERemoteSess
 	else if (InChannelName == FRemoteSessionARCameraChannel::StaticType())
 	{
 		// Skip if there isn't an AR session running
-		if (UARBlueprintLibrary::GetARSessionStatus().Status == EARSessionStatus::Running)
+		if (UARBlueprintLibrary::GetARSessionStatus().Status == EARSessionStatus::Running ||
+			// Add our listener if we are the host
+			InMode == ERemoteSessionChannelMode::Read )
 		{
 			// Client side sending only works on iOS with Android coming in the future
 			bool IsSupported = (InMode == ERemoteSessionChannelMode::Read) || PLATFORM_IOS;
