RemoteSessionApp Crashes

#4.20

The lag & crashes on idle or rotate are related to FBackChannelOSCConnection::Run that keeps on receiving packages when the main thread is busy, eventually crashing the app with to out of memory

diff --git a/Engine/Plugins/Experimental/RemoteSession/Source/RemoteSession/Channels/RemoteSessionFrameBufferChannel.h b/Engine/Plugins/Experimental/RemoteSession/Source/RemoteSession/Channels/RemoteSessionFrameBufferChannel.h
--- a/Engine/Plugins/Experimental/RemoteSession/Source/RemoteSession/Channels/RemoteSessionFrameBufferChannel.h
+++ b/Engine/Plugins/Experimental/RemoteSession/Source/RemoteSession/Channels/RemoteSessionFrameBufferChannel.h
@@ -4,6 +4,9 @@
 
 #include "RemoteSessionChannel.h"
 #include "HAL/ThreadSafeCounter.h"
+#include "HAL/Runnable.h"
+#include "HAL/RunnableThread.h"
+#include "HAL/ThreadSafeBool.h"
 
 
 class FBackChannelOSCMessage;
@@ -16,7 +19,7 @@ class UTexture2D;
  *	A channel that captures the framebuffer on the host, encodes it as a jpg as an async task, then sends it to the client.
  *	On the client images are decoded into a double-buffered texture that can be accessed via GetHostScreen.
  */
-class REMOTESESSION_API FRemoteSessionFrameBufferChannel : public IRemoteSessionChannel
+class REMOTESESSION_API FRemoteSessionFrameBufferChannel : public IRemoteSessionChannel, FRunnable
 {
 public:
 
@@ -87,7 +90,7 @@ protected:
 	TArray<TSharedPtr<FImageData, ESPMode::ThreadSafe>>		IncomingEncodedImages;
 
 	FCriticalSection										DecodedImageMutex;
-	TArray<TSharedPtr<FImageData>>							IncomingDecodedImages;
+	TArray<TSharedPtr<FImageData, ESPMode::ThreadSafe>>		IncomingDecodedImages;
 	FThreadSafeCounter										NumDecodingTasks;
 
 	UTexture2D*												DecodedTextures[2];
@@ -95,7 +98,6 @@ protected:
 
 	/** Time we last sent an image */
 	double LastSentImageTime;
-	int KickedTaskCount;
 	int NumSentImages;
 
 	/** So we can manage callback lifetimes properly */
@@ -106,4 +108,25 @@ protected:
 
 	/** Holds a reference to the scene viewport */
 	TSharedPtr<FSceneViewport> SceneViewport;
+
+	/* Last processed incoming image */
+	int32 LastIncomingImageIndex;
+
+	/* Last processed decode image */
+	int32 LastDecodedImageIndex;
+
+	/** Decode the incoming images on a dedicated thread */
+	void	ProcessIncomingTextures();
+
+protected:
+	
+	/** Background thread for decoding incoming images */
+	uint32 Run();
+	void StartBackgroundThread();
+	void ExitBackgroundThread();
+
+	FRunnableThread*	BackgroundThread;
+	FEvent *			ScreenshotEvent;
+
+	FThreadSafeBool		ExitRequested;
 };
diff --git a/Engine/Plugins/Experimental/RemoteSession/Source/RemoteSession/Private/Channels/RemoteSessionFrameBufferChannel.cpp b/Engine/Plugins/Experimental/RemoteSession/Source/RemoteSession/Private/Channels/RemoteSessionFrameBufferChannel.cpp
--- a/Engine/Plugins/Experimental/RemoteSession/Source/RemoteSession/Private/Channels/RemoteSessionFrameBufferChannel.cpp
+++ b/Engine/Plugins/Experimental/RemoteSession/Source/RemoteSession/Private/Channels/RemoteSessionFrameBufferChannel.cpp
@@ -1,4 +1,4 @@
-﻿// Copyright 1998-2018 Epic Games, Inc. All Rights Reserved.
+// Copyright 1998-2018 Epic Games, Inc. All Rights Reserved.
 
 #include "Channels/RemoteSessionFrameBufferChannel.h"
 #include "RemoteSession.h"
@@ -44,16 +44,22 @@ FRemoteSessionFrameBufferChannel::FRemoteSessionFrameBufferChannel(ERemoteSessio
 	DecodedTextures[1] = nullptr;
 	DecodedTextureIndex = 0;
 	NumSentImages = 0;
-	KickedTaskCount = 0;
 	Role = InRole;
 	ViewportResized = false;
+	LastDecodedImageIndex = 0;
+	LastIncomingImageIndex = 0;
 
+	BackgroundThread = nullptr;
+	ScreenshotEvent = nullptr;
+	ExitRequested = false;
 	if (Role == ERemoteSessionChannelMode::Read)
 	{
 		auto Delegate = FBackChannelDispatchDelegate::FDelegate::CreateRaw(this, &FRemoteSessionFrameBufferChannel::ReceiveHostImage);
 		MessageCallbackHandle = InConnection->AddMessageHandler(TEXT("/Screen"), Delegate);
 
 		InConnection->SetMessageOptions(TEXT("/Screen"), 1);
+
+		StartBackgroundThread();
 	}
 }
 
@@ -68,11 +74,22 @@ FRemoteSessionFrameBufferChannel::~FRemoteSessionFrameBufferChannel()
 			LocalConnection->RemoveMessageHandler(TEXT("/Screen"), MessageCallbackHandle);
 		}
 		MessageCallbackHandle.Reset();
-	}
 
-	while (NumDecodingTasks.GetValue() > 0)
-	{
-		FPlatformProcess::SleepNoStats(0);
+		ExitBackgroundThread();
+
+		if (ScreenshotEvent)
+		{
+			//Cleanup the FEvent
+			FGenericPlatformProcess::ReturnSynchEventToPool(ScreenshotEvent);
+			ScreenshotEvent = nullptr;
+		}
+
+		if (BackgroundThread)
+		{
+			//Cleanup the worker thread
+			delete BackgroundThread;
+			BackgroundThread = nullptr;
+		}
 	}
 
 	ReleaseFrameGrabber();
@@ -191,7 +208,7 @@ void FRemoteSessionFrameBufferChannel::Tick(const float InDeltaTime)
 	{
 		SCOPE_CYCLE_COUNTER(STAT_TextureUpdate);
 
-		TSharedPtr<FImageData> QueuedImage;
+		TSharedPtr<FImageData, ESPMode::ThreadSafe> QueuedImage;
 
 		{
 			// Check to see if there are any queued images. We just care about the last
@@ -200,6 +217,7 @@ void FRemoteSessionFrameBufferChannel::Tick(const float InDeltaTime)
 			{
 				INC_DWORD_STAT(STAT_RSNumFrames);
 				QueuedImage = IncomingDecodedImages.Last();
+				LastDecodedImageIndex = QueuedImage->ImageIndex;
 
 				UE_LOG(LogRemoteSession, Verbose, TEXT("GT: Image %d is ready, discarding %d earlier images"),
 					QueuedImage->ImageIndex, IncomingDecodedImages.Num()-1);
@@ -289,82 +307,78 @@ void FRemoteSessionFrameBufferChannel::ReceiveHostImage(FBackChannelOSCMessage&
 	Message << ReceivedImage->ImageIndex;
 
 	FScopeLock Lock(&IncomingImageMutex);
+	if (LastIncomingImageIndex > 0 && IncomingEncodedImages.Num() > 1 && IncomingEncodedImages.Last()->ImageIndex > LastIncomingImageIndex)
+	{
+		IncomingEncodedImages.RemoveSingle(IncomingEncodedImages.Last());
+	}
 	IncomingEncodedImages.Add(ReceivedImage);
 
+	if (ScreenshotEvent)
+	{
+		// wake up the background thread.
+		ScreenshotEvent->Trigger();
+	}
+
 	UE_LOG(LogRemoteSession, Verbose, TEXT("Received Image %d, %d pending"), 
 		ReceivedImage->ImageIndex, IncomingEncodedImages.Num());
 
-	if (NumDecodingTasks.GetValue() == 0)
+}
+
+void FRemoteSessionFrameBufferChannel::ProcessIncomingTextures()
+{
+	TSharedPtr<FImageData, ESPMode::ThreadSafe> Image;
+	const double StartTime = FPlatformTime::Seconds();
 	{
-		NumDecodingTasks.Increment();
-		KickedTaskCount++;
+		// check if there's anything to do, if not pause the background thread
+		FScopeLock TaskLock(&IncomingImageMutex);
 
-		AsyncTask(ENamedThreads::AnyBackgroundThreadNormalTask, [this]()
+		if (IncomingEncodedImages.Num() == 0)
 		{
-			SCOPE_CYCLE_COUNTER(STAT_ImageDecompression);
+			return;
+		}
 
-			int ProcessedImageCount = 0;
+		// take the last image we don't care about the rest
+		Image = IncomingEncodedImages.Last();
+		LastIncomingImageIndex = Image->ImageIndex;
 
-			do
-			{
-				TSharedPtr<FImageData, ESPMode::ThreadSafe> Image;
+		UE_LOG(LogRemoteSession, Verbose, TEXT("Processing Image %d, discarding %d other pending images"), 
+			Image->ImageIndex, IncomingEncodedImages.Num() - 1);
 
-				const double StartTime = FPlatformTime::Seconds();
-
-				{
-					// check if there's anything to do, if not this task is done
-					FScopeLock TaskLock(&IncomingImageMutex);
+		IncomingEncodedImages.Empty();
+	}
 
-					if (IncomingEncodedImages.Num() == 0)
-					{
-						NumDecodingTasks.Decrement();
-						return;
-					}
+	IImageWrapperModule* ImageWrapperModule = FModuleManager::GetModulePtr<IImageWrapperModule>(FName("ImageWrapper"));
 
-					// take the last image we don't care about the rest
-					Image = IncomingEncodedImages.Last();
+	if (ImageWrapperModule != nullptr)
+	{
+		TSharedPtr<IImageWrapper> ImageWrapper = ImageWrapperModule->CreateImageWrapper(EImageFormat::JPEG);
 
-					UE_LOG(LogRemoteSession, Verbose, TEXT("Processing Image %d, discarding %d other pending images"),
-						Image->ImageIndex, IncomingEncodedImages.Num()-1);
+		ImageWrapper->SetCompressed(Image->ImageData.GetData(), Image->ImageData.Num());
 
-					IncomingEncodedImages.Empty();
-				}
+		const TArray<uint8>* RawData = nullptr;
 
-				IImageWrapperModule* ImageWrapperModule = FModuleManager::GetModulePtr<IImageWrapperModule>(FName("ImageWrapper"));
+		if (ImageWrapper->GetRaw(ERGBFormat::BGRA, 8, RawData))
+		{
+			TSharedPtr<FImageData, ESPMode::ThreadSafe> QueuedImage = MakeShareable(new FImageData);
+			QueuedImage->Width = Image->Width;
+			QueuedImage->Height = Image->Height;
+			QueuedImage->ImageData = MoveTemp(*((TArray<uint8>*)RawData));
+			QueuedImage->ImageIndex = Image->ImageIndex;
 
-				if (ImageWrapperModule != nullptr)
+			{
+				FScopeLock ImageLock(&DecodedImageMutex);
+				if (LastDecodedImageIndex > 0 && IncomingDecodedImages.Num() > 1 && IncomingDecodedImages.Last()->ImageIndex > LastDecodedImageIndex)
 				{
-					TSharedPtr<IImageWrapper> ImageWrapper = ImageWrapperModule->CreateImageWrapper(EImageFormat::JPEG);
-
-					ImageWrapper->SetCompressed(Image->ImageData.GetData(), Image->ImageData.Num());
-
-					const TArray<uint8>* RawData = nullptr;
-
-					if (ImageWrapper->GetRaw(ERGBFormat::BGRA, 8, RawData))
-					{
-						TSharedPtr<FImageData> QueuedImage = MakeShareable(new FImageData);
-						QueuedImage->Width = Image->Width;
-						QueuedImage->Height = Image->Height;
-						QueuedImage->ImageData = MoveTemp(*((TArray<uint8>*)RawData));
-						QueuedImage->ImageIndex = Image->ImageIndex;
-
-						{
-							FScopeLock ImageLock(&DecodedImageMutex);
-							IncomingDecodedImages.Add(QueuedImage);
-
-							UE_LOG(LogRemoteSession, Verbose, TEXT("finished decompressing image %d in %.02f ms (%d in queue)"),
-								Image->ImageIndex,
-								(FPlatformTime::Seconds() - StartTime) * 1000.0,
-								IncomingEncodedImages.Num());
-						}
-					}
+					IncomingDecodedImages.RemoveSingle(IncomingDecodedImages.Last());
 				}
+				IncomingDecodedImages.Add(QueuedImage);
 
-			} while (true);
-
-			UE_LOG(LogRemoteSession, Verbose, TEXT("No remaining images for task %d (%d processed). Exiting."),
-				KickedTaskCount, ProcessedImageCount);
-		});
+				UE_LOG(LogRemoteSession, Verbose, TEXT("finished decompressing image %d in %.02f ms (%d in queue)"),
+					Image->ImageIndex,
+					(FPlatformTime::Seconds() - StartTime) * 1000.0,
+					IncomingEncodedImages.Num());
+			}
+		}
 	}
 }
 
@@ -389,3 +403,42 @@ void FRemoteSessionFrameBufferChannel::OnViewportResized(FVector2D NewSize)
 	ViewportResized = true;
 }
 
+void FRemoteSessionFrameBufferChannel::StartBackgroundThread()
+{
+	check(BackgroundThread == nullptr);
+
+	ExitRequested	= false;
+	
+	ScreenshotEvent = FGenericPlatformProcess::GetSynchEventFromPool(false);;
+
+	BackgroundThread = FRunnableThread::Create(this, TEXT("RemoteSessionFrameBufferThread"), 1024 * 1024, TPri_AboveNormal);
+
+}
+
+void FRemoteSessionFrameBufferChannel::ExitBackgroundThread()
+{
+	ExitRequested = true;
+
+	if (ScreenshotEvent)
+	{
+		ScreenshotEvent->Trigger();
+	}
+
+	if (BackgroundThread)
+	{
+		BackgroundThread->WaitForCompletion();
+	}
+}
+
+uint32 FRemoteSessionFrameBufferChannel::Run()
+{
+	while(!ExitRequested)
+	{
+		// wait a maximum of 1 second or until triggered
+		ScreenshotEvent->Wait(1000);
+
+		ProcessIncomingTextures();
+	}
+
+	return 0;
+}
\ No newline at end of file
