UE-61289 RemoteSession - Unable to reconnect to a PIE session on Mac with an Android device after killing a PIE instance
UE-61330 RemoteSession - Reconnecting to a host can cause the device to get stuck on the last frame of the previous session
UE-61285 Remote Session: Some Android devices do not connect to a Mac host

#Android
#4.20

OnBindEndpoints should be called immediately after receiving "/Version",  moving it to an AsyncTask may cause the client to receive the "/ChannelSelection" signal from the server before the handler for this signal (OnChannelSelection) is set

diff --git a/Engine/Plugins/Experimental/RemoteSession/Source/RemoteSession/Private/RemoteSessionClient.cpp b/Engine/Plugins/Experimental/RemoteSession/Source/RemoteSession/Private/RemoteSessionClient.cpp
--- a/Engine/Plugins/Experimental/RemoteSession/Source/RemoteSession/Private/RemoteSessionClient.cpp
+++ b/Engine/Plugins/Experimental/RemoteSession/Source/RemoteSession/Private/RemoteSessionClient.cpp
@@ -152,9 +152,6 @@ void FRemoteSessionClient::CheckConnection()
 
 void FRemoteSessionClient::OnBindEndpoints()
 {
-	auto Delegate = FBackChannelDispatchDelegate::FDelegate::CreateRaw(this, &FRemoteSessionClient::OnChannelSelection);
-	OSCConnection->AddMessageHandler(*GetChannelSelectionEndPoint(),Delegate);
-	
 	FRemoteSessionRole::OnBindEndpoints();
 }
 
diff --git a/Engine/Plugins/Experimental/RemoteSession/Source/RemoteSession/Private/RemoteSessionClient.h b/Engine/Plugins/Experimental/RemoteSession/Source/RemoteSession/Private/RemoteSessionClient.h
--- a/Engine/Plugins/Experimental/RemoteSession/Source/RemoteSession/Private/RemoteSessionClient.h
+++ b/Engine/Plugins/Experimental/RemoteSession/Source/RemoteSession/Private/RemoteSessionClient.h
@@ -25,7 +25,7 @@ protected:
 	void 			StartConnection();
 	void 			CheckConnection();
 	
-	void OnChannelSelection(FBackChannelOSCMessage& Message, FBackChannelOSCDispatch& Dispatch);
+	void OnChannelSelection(FBackChannelOSCMessage& Message, FBackChannelOSCDispatch& Dispatch) override;
 
 	FString				HostAddress;
 	
diff --git a/Engine/Plugins/Experimental/RemoteSession/Source/RemoteSession/Private/RemoteSessionRole.cpp b/Engine/Plugins/Experimental/RemoteSession/Source/RemoteSession/Private/RemoteSessionRole.cpp
--- a/Engine/Plugins/Experimental/RemoteSession/Source/RemoteSession/Private/RemoteSessionRole.cpp
+++ b/Engine/Plugins/Experimental/RemoteSession/Source/RemoteSession/Private/RemoteSessionRole.cpp
@@ -134,6 +134,9 @@ void FRemoteSessionRole::CreateOSCConnection(TSharedRef<IBackChannelConnection>
 	auto Delegate = FBackChannelDispatchDelegate::FDelegate::CreateRaw(this, &FRemoteSessionRole::OnVersionCheck);
 	OSCConnection->AddMessageHandler(TEXT("/Version"),Delegate);
 	
+	Delegate = FBackChannelDispatchDelegate::FDelegate::CreateRaw(this, &FRemoteSessionRole::OnCreateChannels);
+	OSCConnection->AddMessageHandler(*GetChannelSelectionEndPoint(), Delegate);
+
 	OSCConnection->StartReceiveThread();
 	
 	SendVersion();
@@ -189,6 +192,15 @@ void FRemoteSessionRole::OnVersionCheck(FBackChannelOSCMessage& Message, FBackCh
 	}
 }
 
+void FRemoteSessionRole::OnCreateChannels(FBackChannelOSCMessage& Message, FBackChannelOSCDispatch& Dispatch)
+{
+	OnChannelSelection(Message, Dispatch);
+}
+
+void FRemoteSessionRole::OnChannelSelection(FBackChannelOSCMessage& Message, FBackChannelOSCDispatch& Dispatch)
+{
+}
+
 void FRemoteSessionRole::OnBindEndpoints()
 {
 }
diff --git a/Engine/Plugins/Experimental/RemoteSession/Source/RemoteSession/Private/RemoteSessionRole.h b/Engine/Plugins/Experimental/RemoteSession/Source/RemoteSession/Private/RemoteSessionRole.h
--- a/Engine/Plugins/Experimental/RemoteSession/Source/RemoteSession/Private/RemoteSessionRole.h
+++ b/Engine/Plugins/Experimental/RemoteSession/Source/RemoteSession/Private/RemoteSessionRole.h
@@ -44,9 +44,11 @@ protected:
 	FString			GetVersion() const;
 	void			SendVersion();
 	void 			OnVersionCheck(FBackChannelOSCMessage& Message, FBackChannelOSCDispatch& Dispatch);
+	void			OnCreateChannels(FBackChannelOSCMessage& Message, FBackChannelOSCDispatch& Dispatch);
 	
 	virtual void	OnBindEndpoints();
 	virtual void	OnCreateChannels();
+	virtual void	OnChannelSelection(FBackChannelOSCMessage& Message, FBackChannelOSCDispatch& Dispatch);
 	
 	void 			CreateChannels(const TMap<FString, ERemoteSessionChannelMode>& ChannelMap);
 	void 			CreateChannel(const FString& InChannelList, ERemoteSessionChannelMode InRole);
