UE-61095 ARKit is not working with Virtual Camera

#iOS
#4.20

Reverting the changes from CL4135916 (//Dev-VR, Changed remote session AR dependent channels to check for a running AR session before being created)
- restoring the ARKit initialization in FRemoteSessionXRTrackingChannel & FRemoteSessionARCameraChannel.
- removing the check for a running AR session, before creating these channels in FRemoteSessionRole::CreateChannel

diff --git a/Engine/Plugins/Experimental/RemoteSession/Source/RemoteSession/Private/Channels/RemoteSessionARCameraChannel.cpp b/Engine/Plugins/Experimental/RemoteSession/Source/RemoteSession/Private/Channels/RemoteSessionARCameraChannel.cpp
--- a/Engine/Plugins/Experimental/RemoteSession/Source/RemoteSession/Private/Channels/RemoteSessionARCameraChannel.cpp
+++ b/Engine/Plugins/Experimental/RemoteSession/Source/RemoteSession/Private/Channels/RemoteSessionARCameraChannel.cpp
@@ -378,6 +378,15 @@ FRemoteSessionARCameraChannel::FRemoteSessionARCameraChannel(ERemoteSessionChann
 	, Connection(InConnection)
 	, Role(InRole)
 {
+	
+	static bool OnceTimeARInit = false;
+	if (!OnceTimeARInit && InRole == ERemoteSessionChannelMode::Write)
+	{
+		UARSessionConfig* Config = NewObject<UARSessionConfig>();
+		UARBlueprintLibrary::StartARSession(Config);
+		OnceTimeARInit = true;
+	}
+	
 	RenderingTextures[0] = nullptr;
 	RenderingTextures[1] = nullptr;
 	PPMaterial = LoadObject<UMaterialInterface>(nullptr, TEXT("/RemoteSession/ARCameraPostProcess.ARCameraPostProcess"));
diff --git a/Engine/Plugins/Experimental/RemoteSession/Source/RemoteSession/Private/Channels/RemoteSessionXRTrackingChannel.cpp b/Engine/Plugins/Experimental/RemoteSession/Source/RemoteSession/Private/Channels/RemoteSessionXRTrackingChannel.cpp
--- a/Engine/Plugins/Experimental/RemoteSession/Source/RemoteSession/Private/Channels/RemoteSessionXRTrackingChannel.cpp
+++ b/Engine/Plugins/Experimental/RemoteSession/Source/RemoteSession/Private/Channels/RemoteSessionXRTrackingChannel.cpp
@@ -55,6 +55,16 @@ FRemoteSessionXRTrackingChannel::FRemoteSessionXRTrackingChannel(ERemoteSessionC
 		MessageCallbackHandle = Connection->AddMessageHandler(MESSAGE_ADDRESS, Delegate);
 		Connection->SetMessageOptions(MESSAGE_ADDRESS, 1);
 	}
+    else
+    {
+        // Initialize AR if desired - todo, does this need to check device caps?
+#if PLATFORM_IOS
+		// Workaround - we don't want to set bSupportAR in our project as it prevents us running on old devices, but this needs
+		// to be true before we try to init ARKit stuff
+		UARSessionConfig* Config = NewObject<UARSessionConfig>();
+		UARBlueprintLibrary::StartARSession(Config);
+#endif
+    }
 }
 
 FRemoteSessionXRTrackingChannel::~FRemoteSessionXRTrackingChannel()
diff --git a/Engine/Plugins/Experimental/RemoteSession/Source/RemoteSession/Private/RemoteSessionRole.cpp b/Engine/Plugins/Experimental/RemoteSession/Source/RemoteSession/Private/RemoteSessionRole.cpp
--- a/Engine/Plugins/Experimental/RemoteSession/Source/RemoteSession/Private/RemoteSessionRole.cpp
+++ b/Engine/Plugins/Experimental/RemoteSession/Source/RemoteSession/Private/RemoteSessionRole.cpp
@@ -10,7 +10,6 @@
 #include "Channels/RemoteSessionFrameBufferChannel.h"
 #include "Async/Async.h"
 #include "GeneralProjectSettings.h"
-#include "ARBlueprintLibrary.h"
 
 DEFINE_LOG_CATEGORY(LogRemoteSession);
 
@@ -221,41 +220,21 @@ void FRemoteSessionRole::CreateChannel(const FString& InChannelName, ERemoteSess
 	{
 		NewChannel = MakeShareable(new FRemoteSessionFrameBufferChannel(InMode, OSCConnection));
 	}
-    else if (InChannelName == FRemoteSessionXRTrackingChannel::StaticType())
+	else if (InChannelName == FRemoteSessionXRTrackingChannel::StaticType())
 	{
-		// Skip if there isn't an AR session running
-		if (UARBlueprintLibrary::GetARSessionStatus().Status == EARSessionStatus::Running ||
-			// Add our listener if we are the host
-			InMode == ERemoteSessionChannelMode::Read)
-		{
-			NewChannel = MakeShareable(new FRemoteSessionXRTrackingChannel(InMode, OSCConnection));
-		}
-		else
-		{
-			UE_LOG(LogRemoteSession, Warning, TEXT("No AR session in progress so FRemoteSessionXRTrackingChannel is disabled"));
-		}
+		NewChannel = MakeShareable(new FRemoteSessionXRTrackingChannel(InMode, OSCConnection));
 	}
 	else if (InChannelName == FRemoteSessionARCameraChannel::StaticType())
 	{
-		// Skip if there isn't an AR session running
-		if (UARBlueprintLibrary::GetARSessionStatus().Status == EARSessionStatus::Running ||
-			// Add our listener if we are the host
-			InMode == ERemoteSessionChannelMode::Read )
+		// Client side sending only works on iOS with Android coming in the future
+		bool IsSupported = (InMode == ERemoteSessionChannelMode::Read) || PLATFORM_IOS;
+		if (IsSupported)
 		{
-			// Client side sending only works on iOS with Android coming in the future
-			bool IsSupported = (InMode == ERemoteSessionChannelMode::Read) || PLATFORM_IOS;
-			if (IsSupported)
-			{
-				NewChannel = MakeShareable(new FRemoteSessionARCameraChannel(InMode, OSCConnection));
-			}
-			else
-			{
-				UE_LOG(LogRemoteSession, Warning, TEXT("FRemoteSessionARCameraChannel does not support sending on this platform"));
-			}
+			NewChannel = MakeShareable(new FRemoteSessionARCameraChannel(InMode, OSCConnection));
 		}
 		else
 		{
-			UE_LOG(LogRemoteSession, Warning, TEXT("No AR session in progress so FRemoteSessionARCameraChannel is disabled"));
+			UE_LOG(LogRemoteSession, Warning, TEXT("FRemoteSessionARCameraChannel does not support sending on this platform"));
 		}
 	}
 	
