UE-63334 Unreal Remote 2 does not send information about motion control

#iOS
#Android

The functionality was commented out. Tested it for performance and it seems OK (despite the high frequency of /motion signals)

diff --git a/Engine/Plugins/Experimental/RemoteSession/Source/RemoteSession/Private/MessageHandler/RecordingMessageHandler.cpp b/Engine/Plugins/Experimental/RemoteSession/Source/RemoteSession/Private/MessageHandler/RecordingMessageHandler.cpp
--- a/Engine/Plugins/Experimental/RemoteSession/Source/RemoteSession/Private/MessageHandler/RecordingMessageHandler.cpp
+++ b/Engine/Plugins/Experimental/RemoteSession/Source/RemoteSession/Private/MessageHandler/RecordingMessageHandler.cpp
@@ -410,7 +410,7 @@ void FRecordingMessageHandler::PlayOnEndGesture(FArchive& Ar)
 
 bool FRecordingMessageHandler::OnMotionDetected(const FVector& Tilt, const FVector& RotationRate, const FVector& Gravity, const FVector& Acceleration, int32 ControllerId)
 {
-	/*if (IsRecording())
+	if (IsRecording())
 	{
 		FiveParamMsg<FVector, FVector, FVector, FVector, int32> 
 			Msg(Tilt, RotationRate, Gravity, Acceleration, ControllerId);
@@ -420,7 +420,7 @@ bool FRecordingMessageHandler::OnMotionDetected(const FVector& Tilt, const FVect
 	if (ConsumeInput)
 	{
 		return true;
-	}*/
+	}
 
 	return FProxyMessageHandler::OnMotionDetected(Tilt, RotationRate, Gravity, Acceleration, ControllerId);
 }
