Stop the frame grabber as soon as the remove session end.

diff --git a/Engine/Plugins/Experimental/RemoteSession/Source/RemoteSession/Private/Channels/RemoteSessionFrameBufferChannel.cpp b/Engine/Plugins/Experimental/RemoteSession/Source/RemoteSession/Private/Channels/RemoteSessionFrameBufferChannel.cpp
--- a/Engine/Plugins/Experimental/RemoteSession/Source/RemoteSession/Private/Channels/RemoteSessionFrameBufferChannel.cpp
+++ b/Engine/Plugins/Experimental/RemoteSession/Source/RemoteSession/Private/Channels/RemoteSessionFrameBufferChannel.cpp
@@ -108,7 +108,7 @@ void FRemoteSessionFrameBufferChannel::ReleaseFrameGrabber()
 {
 	if (FrameGrabber.IsValid())
 	{
-		FrameGrabber->StopCapturingFrames();
+		FrameGrabber->Shutdown();
 		FrameGrabber = nullptr;
 	}
 }
