Added ability for RemoteSession to send controller input from client device to host for gamepad testing

#tests build on PC & Mac, and iOS.

diff --git a/Engine/Plugins/Experimental/RemoteSession/Source/RemoteSession/Channels/RemoteSessionInputChannel.h b/Engine/Plugins/Experimental/RemoteSession/Source/RemoteSession/Channels/RemoteSessionInputChannel.h
--- a/Engine/Plugins/Experimental/RemoteSession/Source/RemoteSession/Channels/RemoteSessionInputChannel.h
+++ b/Engine/Plugins/Experimental/RemoteSession/Source/RemoteSession/Channels/RemoteSessionInputChannel.h
@@ -31,7 +31,7 @@ public:
 
 protected:
 
-	TWeakPtr<FGenericApplicationMessageHandler> DefaultHandler;
+	TSharedPtr<FGenericApplicationMessageHandler> DefaultHandler;
 
 	TSharedPtr<FRecordingMessageHandler> RecordingHandler;
 
diff --git a/Engine/Plugins/Experimental/RemoteSession/Source/RemoteSession/Private/Channels/RemoteSessionInputChannel.cpp b/Engine/Plugins/Experimental/RemoteSession/Source/RemoteSession/Private/Channels/RemoteSessionInputChannel.cpp
--- a/Engine/Plugins/Experimental/RemoteSession/Source/RemoteSession/Private/Channels/RemoteSessionInputChannel.cpp
+++ b/Engine/Plugins/Experimental/RemoteSession/Source/RemoteSession/Private/Channels/RemoteSessionInputChannel.cpp
@@ -7,6 +7,10 @@
 #include "MessageHandler/RecordingMessageHandler.h"
 
 
+namespace RemoteSessionVars
+{
+	static FAutoConsoleVariable BlockLocalInput(TEXT("Remote.BlockLocalInput"), 0, TEXT("Don't accept local input when a host is connected"));
+};
 
 FRemoteSessionInputChannel::FRemoteSessionInputChannel(ERemoteSessionChannelMode InRole, TSharedPtr<FBackChannelOSCConnection, ESPMode::ThreadSafe> InConnection)
 	: IRemoteSessionChannel(InRole, InConnection)
@@ -19,22 +23,25 @@ FRemoteSessionInputChannel::FRemoteSessionInputChannel(ERemoteSessionChannelMode
 	// handler for that data 
 	if (Role == ERemoteSessionChannelMode::Write)
 	{
+		// when we're recording (writing) we create a recording handler that receives platform-level input
+		// and passes it through to the default UE handler
 		DefaultHandler = FSlateApplication::Get().GetPlatformApplication()->GetMessageHandler();
 
-		RecordingHandler = MakeShareable(new FRecordingMessageHandler(DefaultHandler.Pin()));
-
+		RecordingHandler = MakeShareable(new FRecordingMessageHandler(DefaultHandler));
 		RecordingHandler->SetRecordingHandler(this);
 
 		FSlateApplication::Get().GetPlatformApplication()->SetMessageHandler(RecordingHandler.ToSharedRef());
 	}
 	else
 	{
-		TSharedRef<FGenericApplicationMessageHandler> DestinationHandler = FSlateApplication::Get().GetPlatformApplication()->GetMessageHandler();
+		DefaultHandler = FSlateApplication::Get().GetPlatformApplication()->GetMessageHandler();
 
-		PlaybackHandler = MakeShareable(new FRecordingMessageHandler(DestinationHandler));
+		PlaybackHandler = MakeShareable(new FRecordingMessageHandler(DefaultHandler));
 		
 		auto Delegate = FBackChannelDispatchDelegate::FDelegate::CreateRaw(this, &FRemoteSessionInputChannel::OnRemoteMessage);
 		MessageCallbackHandle = Connection->AddMessageHandler(TEXT("/MessageHandler/"), Delegate);
+
+		FSlateApplication::Get().GetPlatformApplication()->SetMessageHandler(PlaybackHandler.ToSharedRef());
 	}
 }
 
@@ -50,7 +57,7 @@ FRemoteSessionInputChannel::~FRemoteSessionInputChannel()
 	// todo - is this ok? Might other things have changed the handler like we do?
 	if (DefaultHandler.IsValid())
 	{
-		FSlateApplication::Get().GetPlatformApplication()->SetMessageHandler(DefaultHandler.Pin().ToSharedRef());
+		FSlateApplication::Get().GetPlatformApplication()->SetMessageHandler(DefaultHandler.ToSharedRef());
 	}
 
 	// should restore handler? What if something else changed it...
@@ -76,7 +83,20 @@ void FRemoteSessionInputChannel::SetInputRect(const FVector2D& TopLeft, const FV
 
 void FRemoteSessionInputChannel::Tick(const float InDeltaTime)
 {
-	// everything happens via messaging.
+	// everything else happens via messaging.
+	if (Role == ERemoteSessionChannelMode::Read)
+	{
+		bool bBlockInput = RemoteSessionVars::BlockLocalInput->GetInt() > 0;
+
+		if (bBlockInput && !PlaybackHandler->IsConsumingInput())
+		{
+			PlaybackHandler->SetConsumeInput(true);
+		}
+		else if (!bBlockInput && PlaybackHandler->IsConsumingInput())
+		{
+			PlaybackHandler->SetConsumeInput(false);
+		}
+	}
 }
 
 void FRemoteSessionInputChannel::RecordMessage(const TCHAR* MsgName, const TArray<uint8>& Data)
@@ -100,6 +120,6 @@ void FRemoteSessionInputChannel::OnRemoteMessage(FBackChannelOSCMessage& Message
 
 	TArray<uint8> MsgData;
 	Message << MsgData;
-
+	
 	PlaybackHandler->PlayMessage(*MessageName, MsgData);
 }
diff --git a/Engine/Plugins/Experimental/RemoteSession/Source/RemoteSession/Private/MessageHandler/ProxyMessageHandler.cpp b/Engine/Plugins/Experimental/RemoteSession/Source/RemoteSession/Private/MessageHandler/ProxyMessageHandler.cpp
--- a/Engine/Plugins/Experimental/RemoteSession/Source/RemoteSession/Private/MessageHandler/ProxyMessageHandler.cpp
+++ b/Engine/Plugins/Experimental/RemoteSession/Source/RemoteSession/Private/MessageHandler/ProxyMessageHandler.cpp
@@ -8,10 +8,21 @@ FProxyMessageHandler::FProxyMessageHandler(const TSharedPtr<FGenericApplicationM
 
 }
 
-FProxyMessageHandler::~FProxyMessageHandler() 
+FProxyMessageHandler::~FProxyMessageHandler()
 {
 }
 
+void FProxyMessageHandler::SetTargetHandler(const TSharedPtr<FGenericApplicationMessageHandler>& InTargetHandler)
+{
+	TargetHandler = InTargetHandler;
+}
+
+const TSharedPtr<FGenericApplicationMessageHandler> FProxyMessageHandler::GetTargetHandler() const
+{
+	return TargetHandler;
+}
+
+
 bool FProxyMessageHandler::ShouldProcessUserInputMessages(const TSharedPtr< FGenericWindow >& PlatformWindow) const
 {
 	if (TargetHandler.IsValid())
@@ -164,16 +175,31 @@ bool FProxyMessageHandler::OnCursorSet()
 
 bool FProxyMessageHandler::OnControllerAnalog(FGamepadKeyNames::Type KeyName, int32 ControllerId, float AnalogValue)
 {
+	if (TargetHandler.IsValid())
+	{
+		return TargetHandler->OnControllerAnalog(KeyName, ControllerId, AnalogValue);
+	}
+
 	return false;
 }
 
 bool FProxyMessageHandler::OnControllerButtonPressed(FGamepadKeyNames::Type KeyName, int32 ControllerId, bool IsRepeat)
 {
+	if (TargetHandler.IsValid())
+	{
+		return TargetHandler->OnControllerButtonPressed(KeyName, ControllerId, IsRepeat);
+	}
+
 	return false;
 }
 
 bool FProxyMessageHandler::OnControllerButtonReleased(FGamepadKeyNames::Type KeyName, int32 ControllerId, bool IsRepeat)
 {
+	if (TargetHandler.IsValid())
+	{
+		return TargetHandler->OnControllerButtonReleased(KeyName, ControllerId, IsRepeat);
+	}
+
 	return false;
 }
 
diff --git a/Engine/Plugins/Experimental/RemoteSession/Source/RemoteSession/Private/MessageHandler/ProxyMessageHandler.h b/Engine/Plugins/Experimental/RemoteSession/Source/RemoteSession/Private/MessageHandler/ProxyMessageHandler.h
--- a/Engine/Plugins/Experimental/RemoteSession/Source/RemoteSession/Private/MessageHandler/ProxyMessageHandler.h
+++ b/Engine/Plugins/Experimental/RemoteSession/Source/RemoteSession/Private/MessageHandler/ProxyMessageHandler.h
@@ -5,15 +5,24 @@
 #include "CoreMinimal.h"
 #include "GenericPlatform/GenericApplicationMessageHandler.h"
 
+/*
+	Class that can act as a proxy for places where FGenericApplicationMessageHandler is used.
+*/
 class FProxyMessageHandler : public FGenericApplicationMessageHandler
 {
 	
 public:
 
+	FProxyMessageHandler();
+
 	FProxyMessageHandler(const TSharedPtr<FGenericApplicationMessageHandler>& InTargetHandler);
 
 	virtual ~FProxyMessageHandler();
 
+	void SetTargetHandler(const TSharedPtr<FGenericApplicationMessageHandler>& InTargetHandler);
+
+	const TSharedPtr<FGenericApplicationMessageHandler> GetTargetHandler() const;
+
 	virtual bool ShouldProcessUserInputMessages(const TSharedPtr< FGenericWindow >& PlatformWindow) const override;
 
 	virtual bool OnKeyChar(const TCHAR Character, const bool IsRepeat) override;
@@ -112,6 +121,6 @@ public:
 
 protected:
 
-	const TSharedPtr<FGenericApplicationMessageHandler> TargetHandler;
+	TSharedPtr<FGenericApplicationMessageHandler> TargetHandler;
 
 };
\ No newline at end of file
diff --git a/Engine/Plugins/Experimental/RemoteSession/Source/RemoteSession/Private/MessageHandler/RecordingMessageHandler.cpp b/Engine/Plugins/Experimental/RemoteSession/Source/RemoteSession/Private/MessageHandler/RecordingMessageHandler.cpp
--- a/Engine/Plugins/Experimental/RemoteSession/Source/RemoteSession/Private/MessageHandler/RecordingMessageHandler.cpp
+++ b/Engine/Plugins/Experimental/RemoteSession/Source/RemoteSession/Private/MessageHandler/RecordingMessageHandler.cpp
@@ -37,7 +37,7 @@ FRecordingMessageHandler::FRecordingMessageHandler(const TSharedPtr<FGenericAppl
 	: FProxyMessageHandler(InTargetHandler)
 {
 	OutputWriter = nullptr;
-	ConsumeInput = false;
+	bConsumeInput = false;
 	bIsTouching = false;
 	InputRect = FRect(EForceInit::ForceInitToZero);
 	LastTouchLocation = FVector2D(EForceInit::ForceInitToZero);
@@ -56,6 +56,10 @@ FRecordingMessageHandler::FRecordingMessageHandler(const TSharedPtr<FGenericAppl
 	BIND_PLAYBACK_HANDLER(TEXT("OnTouchGesture"), PlayOnTouchGesture);
 	BIND_PLAYBACK_HANDLER(TEXT("OnEndGesture"), PlayOnEndGesture);
 	BIND_PLAYBACK_HANDLER(TEXT("OnTouchForceChanged"), PlayOnTouchForceChanged);
+
+	BIND_PLAYBACK_HANDLER(TEXT("OnControllerAnalog"), PlayOnControllerAnalog);
+	BIND_PLAYBACK_HANDLER(TEXT("OnControllerButtonPressed"), PlayOnControllerButtonPressed);
+	BIND_PLAYBACK_HANDLER(TEXT("OnControllerButtonReleased"), PlayOnControllerButtonReleased);
 }
 
 #undef BIND_PLAYBACK_HANDLER
@@ -73,11 +77,6 @@ void FRecordingMessageHandler::RecordMessage(const TCHAR* MsgName, const TArray<
 	}
 }
 
-void FRecordingMessageHandler::SetConsumeInput(bool bConsume)
-{
-	ConsumeInput = bConsume;
-}
-
 void FRecordingMessageHandler::SetPlaybackWindow(TWeakPtr<SWindow> InWindow, TWeakPtr<FSceneViewport> InViewport)
 {
 	PlaybackWindow = InWindow;
@@ -160,9 +159,18 @@ bool FRecordingMessageHandler::PlayMessage(const TCHAR* Message, const TArray<ui
 		// todo - can we steal this data in a more elegant way? :)
 		TSharedPtr<TArray<uint8>, ESPMode::ThreadSafe> DataCopy = MakeShareable(new TArray<uint8>(MoveTemp(*(TArray<uint8>*)&Data)));
 
-		AsyncTask(ENamedThreads::GameThread, [Dispatch, DataCopy] {
+		AsyncTask(ENamedThreads::GameThread, [this, Dispatch, DataCopy] {
+
+			bool WasBlocking = IsConsumingInput();
+
+			if (WasBlocking)
+			{
+				SetConsumeInput(false);
+			}
 			FMemoryReader Ar(*DataCopy);
 			Dispatch->ExecuteIfBound(Ar);
+
+			SetConsumeInput(WasBlocking);
 		});
 		
 	}
@@ -182,7 +190,7 @@ bool FRecordingMessageHandler::OnKeyChar(const TCHAR Character, const bool IsRep
 		RecordMessage(TEXT("OnKeyChar"), Msg.AsData());
 	}
 
-	if (ConsumeInput)
+	if (bConsumeInput)
 	{
 		return true;
 	}	
@@ -204,7 +212,7 @@ bool FRecordingMessageHandler::OnKeyDown(const int32 KeyCode, const uint32 Chara
 		RecordMessage(TEXT("OnKeyDown"), Msg.AsData());
 	}
 
-	if (ConsumeInput)
+	if (bConsumeInput)
 	{
 		return true;
 	}
@@ -227,7 +235,7 @@ bool FRecordingMessageHandler::OnKeyUp(const int32 KeyCode, const uint32 Charact
 		RecordMessage(TEXT("OnKeyUp"), Msg.AsData());
 	}
 
-	if (ConsumeInput)
+	if (bConsumeInput)
 	{
 		return true;
 	}
@@ -257,7 +265,7 @@ bool FRecordingMessageHandler::OnTouchStarted(const TSharedPtr< FGenericWindow >
 		}
 	}
 
-	if (ConsumeInput)
+	if (bConsumeInput)
 	{
 		return true;
 	}
@@ -297,7 +305,7 @@ bool FRecordingMessageHandler::OnTouchMoved(const FVector2D& Location, float For
 		}
 	}
 
-	if (ConsumeInput)
+	if (bConsumeInput)
 	{
 		return true;
 	}
@@ -330,7 +338,7 @@ bool FRecordingMessageHandler::OnTouchEnded(const FVector2D& Location, int32 Tou
 		bIsTouching = false;
 	}
 
-	if (ConsumeInput)
+	if (bConsumeInput)
 	{
 		return true;
 	}
@@ -362,7 +370,7 @@ bool FRecordingMessageHandler::OnTouchForceChanged(const FVector2D& Location, fl
 		}
 	}
 
-	if (ConsumeInput)
+	if (bConsumeInput)
 	{
 		return true;
 	}
@@ -393,7 +401,7 @@ bool FRecordingMessageHandler::OnTouchFirstMove(const FVector2D& Location, float
 		}
 	}
 
-	if (ConsumeInput)
+	if (bConsumeInput)
 	{
 		return true;
 	}
@@ -418,7 +426,7 @@ void FRecordingMessageHandler::OnBeginGesture()
 		OutputWriter->RecordMessage(TEXT("OnBeginGesture"), Msg.AsData());
 	}
 
-	if (ConsumeInput)
+	if (bConsumeInput)
 	{
 		return;
 	}
@@ -439,7 +447,7 @@ bool FRecordingMessageHandler::OnTouchGesture(EGestureEvent GestureType, const F
 		OutputWriter->RecordMessage(TEXT("OnTouchGesture"), Msg.AsData());
 	}
 
-	if (ConsumeInput)
+	if (bConsumeInput)
 	{
 		return true;
 	}
@@ -461,7 +469,7 @@ void FRecordingMessageHandler::OnEndGesture()
 		OutputWriter->RecordMessage(TEXT("OnEndGesture"), Msg.AsData());
 	}
 
-	if (ConsumeInput)
+	if (bConsumeInput)
 	{
 		return;
 	}
@@ -484,7 +492,7 @@ bool FRecordingMessageHandler::OnMotionDetected(const FVector& Tilt, const FVect
 		OutputWriter->RecordMessage(TEXT("OnMotionDetected"), Msg.AsData());
 	}
 
-	if (ConsumeInput)
+	if (bConsumeInput)
 	{
 		return true;
 	}
@@ -497,3 +505,70 @@ void FRecordingMessageHandler::PlayOnMotionDetected(FArchive& Ar)
 	FiveParamMsg<FVector, FVector, FVector, FVector, int32 > Msg(Ar);
 	OnMotionDetected(Msg.Param1, Msg.Param2, Msg.Param3, Msg.Param4, Msg.Param5);
 }
+
+bool FRecordingMessageHandler::OnControllerAnalog(FGamepadKeyNames::Type KeyName, int32 ControllerId, float AnalogValue)
+{
+	if (IsRecording())
+	{
+		ThreeParamMsg<FString, int32, float> Msg(KeyName.ToString(), ControllerId, AnalogValue);
+		RecordMessage(TEXT("OnControllerAnalog"), Msg.AsData());
+	}
+
+	if (bConsumeInput)
+	{
+		return true;
+	}
+
+	return FProxyMessageHandler::OnControllerAnalog(KeyName, ControllerId, AnalogValue);
+}
+
+void FRecordingMessageHandler::PlayOnControllerAnalog(FArchive& Ar)
+{
+	ThreeParamMsg<FString, int32, float > Msg(Ar);
+	OnControllerAnalog(FName(*Msg.Param1), Msg.Param2, Msg.Param3);
+}
+
+bool FRecordingMessageHandler::OnControllerButtonPressed(FGamepadKeyNames::Type KeyName, int32 ControllerId, bool IsRepeat)
+{
+	if (IsRecording())
+	{
+		ThreeParamMsg<FString, int32, bool> Msg(KeyName.ToString(), ControllerId, IsRepeat);
+		RecordMessage(TEXT("OnControllerButtonPressed"), Msg.AsData());
+	}
+
+	if (bConsumeInput)
+	{
+		return true;
+	}
+
+	return FProxyMessageHandler::OnControllerButtonPressed(KeyName, ControllerId, IsRepeat);
+}
+
+void FRecordingMessageHandler::PlayOnControllerButtonPressed(FArchive& Ar)
+{
+	ThreeParamMsg<FString, int32, bool > Msg(Ar);
+	OnControllerButtonPressed(FName(*Msg.Param1), Msg.Param2, Msg.Param3);
+}
+
+
+bool FRecordingMessageHandler::OnControllerButtonReleased(FGamepadKeyNames::Type KeyName, int32 ControllerId, bool IsRepeat)
+{
+	if (IsRecording())
+	{
+		ThreeParamMsg<FString, int32, bool> Msg(KeyName.ToString(), ControllerId, IsRepeat);
+		RecordMessage(TEXT("OnControllerButtonReleased"), Msg.AsData());
+	}
+
+	if (bConsumeInput)
+	{
+		return true;
+	}
+
+	return FProxyMessageHandler::OnControllerButtonReleased(KeyName, ControllerId, IsRepeat);
+}
+
+void FRecordingMessageHandler::PlayOnControllerButtonReleased(FArchive& Ar)
+{
+	ThreeParamMsg<FString, int32, bool > Msg(Ar);
+	OnControllerButtonReleased(FName(*Msg.Param1), Msg.Param2, Msg.Param3);
+}
diff --git a/Engine/Plugins/Experimental/RemoteSession/Source/RemoteSession/Private/MessageHandler/RecordingMessageHandler.h b/Engine/Plugins/Experimental/RemoteSession/Source/RemoteSession/Private/MessageHandler/RecordingMessageHandler.h
--- a/Engine/Plugins/Experimental/RemoteSession/Source/RemoteSession/Private/MessageHandler/RecordingMessageHandler.h
+++ b/Engine/Plugins/Experimental/RemoteSession/Source/RemoteSession/Private/MessageHandler/RecordingMessageHandler.h
@@ -60,12 +60,20 @@ public:
 
 	void SetRecordingHandler(IRecordingMessageHandlerWriter* InOutputWriter);
 
-	void SetConsumeInput(bool bConsume);
+	/**
+	 * Do not pass input to the target handler
+	 */
+	void SetConsumeInput(bool bConsume) { bConsumeInput = bConsume; }
 
-	bool IsRecording() const
-	{
-		return OutputWriter != nullptr;
-	}
+	/**
+	 * Returns true/false depending on whether input is blocked
+	 */
+	bool IsConsumingInput() const { return bConsumeInput; }
+
+	/**
+	 * Returns true/false depending on whether we are recording (RecordingHandler is != null)
+	 */
+	bool IsRecording() const {	return OutputWriter != nullptr; }
 
 	void SetPlaybackWindow(TWeakPtr<SWindow> InWindow, TWeakPtr<FSceneViewport> InViewport);
 
@@ -73,14 +81,23 @@ public:
 
 public:
 
+	/**
+	 * Keyboard handling
+	 */
 	virtual bool OnKeyChar(const TCHAR Character, const bool IsRepeat) override;
 	virtual bool OnKeyDown(const int32 KeyCode, const uint32 CharacterCode, const bool IsRepeat) override;
 	virtual bool OnKeyUp(const int32 KeyCode, const uint32 CharacterCode, const bool IsRepeat) override;
 
+	/**
+	 * High-level gesture events
+	 */
 	virtual void OnBeginGesture() override;
 	virtual bool OnTouchGesture(EGestureEvent GestureType, const FVector2D& Delta, float WheelDelta, bool bIsDirectionInvertedFromDevice) override;
 	virtual void OnEndGesture() override;
 
+	/**
+	 * Raw touch events
+	 */
 	virtual bool OnTouchStarted(const TSharedPtr< FGenericWindow >& Window, const FVector2D& Location, float Force, int32 TouchIndex, int32 ControllerId) override;
 	virtual bool OnTouchMoved(const FVector2D& Location, float Force, int32 TouchIndex, int32 ControllerId) override;
 	virtual bool OnTouchEnded(const FVector2D& Location, int32 TouchIndex, int32 ControllerId) override;
@@ -88,6 +105,13 @@ public:
 	virtual bool OnTouchFirstMove(const FVector2D& Location, float Force, int32 TouchIndex, int32 ControllerId) override;
 	virtual bool OnMotionDetected(const FVector& Tilt, const FVector& RotationRate, const FVector& Gravity, const FVector& Acceleration, int32 ControllerId) override;
 
+	/**
+	 * Controller handling
+	 */
+	virtual bool OnControllerAnalog(FGamepadKeyNames::Type KeyName, int32 ControllerId, float AnalogValue) override;
+	virtual bool OnControllerButtonPressed(FGamepadKeyNames::Type KeyName, int32 ControllerId, bool IsRepeat) override;
+	virtual bool OnControllerButtonReleased(FGamepadKeyNames::Type KeyName, int32 ControllerId, bool IsRepeat) override;
+
 	bool PlayMessage(const TCHAR* Message, const TArray<uint8>& Data);
 
 protected:
@@ -100,6 +124,7 @@ protected:
 	virtual void PlayOnKeyChar(FArchive& Ar);
 	virtual void PlayOnKeyDown(FArchive& Ar);
 	virtual void PlayOnKeyUp(FArchive& Ar);
+
 	virtual void PlayOnBeginGesture(FArchive& Ar);
 	virtual void PlayOnTouchGesture(FArchive& Ar);
 	virtual void PlayOnEndGesture(FArchive& Ar);
@@ -111,9 +136,13 @@ protected:
 	virtual void PlayOnTouchFirstMove(FArchive& Ar);
 	virtual void PlayOnMotionDetected(FArchive& Ar);
 
+	virtual void PlayOnControllerAnalog(FArchive& Ar);
+	virtual void PlayOnControllerButtonPressed(FArchive& Ar);
+	virtual void PlayOnControllerButtonReleased(FArchive& Ar);
+
 
 	IRecordingMessageHandlerWriter*		OutputWriter;
-	bool								ConsumeInput;
+	bool								bConsumeInput;
 	TWeakPtr<SWindow>					PlaybackWindow;
 	TWeakPtr<FSceneViewport>			PlaybackViewport;
 
diff --git a/Engine/Plugins/Experimental/RemoteSession/Source/RemoteSession/Private/RemoteSessionHost.cpp b/Engine/Plugins/Experimental/RemoteSession/Source/RemoteSession/Private/RemoteSessionHost.cpp
--- a/Engine/Plugins/Experimental/RemoteSession/Source/RemoteSession/Private/RemoteSessionHost.cpp
+++ b/Engine/Plugins/Experimental/RemoteSession/Source/RemoteSession/Private/RemoteSessionHost.cpp
@@ -62,13 +62,6 @@ void FRemoteSessionHost::SetScreenSharing(const bool bEnabled)
 {
 }
 
-void FRemoteSessionHost::SetConsumeInput(const bool bConsume)
-{
-	/*if (PlaybackMessageHandler.IsValid())
-	{
-		PlaybackMessageHandler->SetConsumeInput(bConsume);
-	}*/
-}
 
 bool FRemoteSessionHost::StartListening(const uint16 InPort)
 {
diff --git a/Engine/Plugins/Experimental/RemoteSession/Source/RemoteSession/Private/RemoteSessionHost.h b/Engine/Plugins/Experimental/RemoteSession/Source/RemoteSession/Private/RemoteSessionHost.h
--- a/Engine/Plugins/Experimental/RemoteSession/Source/RemoteSession/Private/RemoteSessionHost.h
+++ b/Engine/Plugins/Experimental/RemoteSession/Source/RemoteSession/Private/RemoteSessionHost.h
@@ -10,6 +10,7 @@ class FFrameGrabber;
 class IImageWrapper;
 class FRemoteSessionInputChannel;
 
+
 class FRemoteSessionHost : public FRemoteSessionRole, public TSharedFromThis<FRemoteSessionHost>
 {
 public:
@@ -23,8 +24,6 @@ public:
 
 	void SetScreenSharing(const bool bEnabled);
 
-	void SetConsumeInput(const bool bConsume);
-
 	virtual void Tick(float DeltaTime) override;
 
 protected:
