RemoteSession / VirtualCamera fixes
- Mapped OnTouchFirstMove / OnTouchForceChanged handlers
- Removed AR input from VCam

diff --git a/Engine/Plugins/Experimental/RemoteSession/Source/RemoteSession/Private/MessageHandler/RecordingMessageHandler.cpp b/Engine/Plugins/Experimental/RemoteSession/Source/RemoteSession/Private/MessageHandler/RecordingMessageHandler.cpp
--- a/Engine/Plugins/Experimental/RemoteSession/Source/RemoteSession/Private/MessageHandler/RecordingMessageHandler.cpp
+++ b/Engine/Plugins/Experimental/RemoteSession/Source/RemoteSession/Private/MessageHandler/RecordingMessageHandler.cpp
@@ -48,12 +48,14 @@ FRecordingMessageHandler::FRecordingMessageHandler(const TSharedPtr<FGenericAppl
 	BIND_PLAYBACK_HANDLER(TEXT("OnKeyDown"), PlayOnKeyDown);
 
 	BIND_PLAYBACK_HANDLER(TEXT("OnTouchStarted"), PlayOnTouchStarted);
+	BIND_PLAYBACK_HANDLER(TEXT("OnTouchFirstMove"), PlayOnTouchFirstMove);
 	BIND_PLAYBACK_HANDLER(TEXT("OnTouchMoved"), PlayOnTouchMoved);
 	BIND_PLAYBACK_HANDLER(TEXT("OnTouchEnded"), PlayOnTouchEnded);
 	BIND_PLAYBACK_HANDLER(TEXT("OnMotionDetected"), PlayOnMotionDetected);
 	BIND_PLAYBACK_HANDLER(TEXT("OnBeginGesture"), PlayOnBeginGesture);
 	BIND_PLAYBACK_HANDLER(TEXT("OnTouchGesture"), PlayOnTouchGesture);
 	BIND_PLAYBACK_HANDLER(TEXT("OnEndGesture"), PlayOnEndGesture);
+	BIND_PLAYBACK_HANDLER(TEXT("OnTouchForceChanged"), PlayOnTouchForceChanged);
 }
 
 #undef BIND_PLAYBACK_HANDLER
