Fixing crash in remoteviewer when SceneViewport doesn't exist yet.

diff --git a/Engine/Plugins/Experimental/RemoteSession/Source/RemoteSession/Private/RemoteSessionHost.cpp b/Engine/Plugins/Experimental/RemoteSession/Source/RemoteSession/Private/RemoteSessionHost.cpp
--- a/Engine/Plugins/Experimental/RemoteSession/Source/RemoteSession/Private/RemoteSessionHost.cpp
+++ b/Engine/Plugins/Experimental/RemoteSession/Source/RemoteSession/Private/RemoteSessionHost.cpp
@@ -147,7 +147,7 @@ void FRemoteSessionHost::OnCreateChannels()
 	
 	// setup framebuffer capture
 	TSharedPtr<FRemoteSessionFrameBufferChannel> FBChannel = IRemoteSessionRole::GetChannel<FRemoteSessionFrameBufferChannel>();
-	if (FBChannel.IsValid())
+	if (FBChannel.IsValid() && SceneViewport.IsValid())
 	{
 		FBChannel->SetCaptureViewport(SceneViewport.ToSharedRef());
 		FBChannel->SetCaptureQuality(Quality, Framerate);
