Fixup fallout from removing include cycle in CL# 9424062
#rnx

diff --git a/Engine/Plugins/Experimental/RemoteSession/Source/RemoteSession/Channels/RemoteSessionChannel.h b/Engine/Plugins/Experimental/RemoteSession/Source/RemoteSession/Channels/RemoteSessionChannel.h
--- a/Engine/Plugins/Experimental/RemoteSession/Source/RemoteSession/Channels/RemoteSessionChannel.h
+++ b/Engine/Plugins/Experimental/RemoteSession/Source/RemoteSession/Channels/RemoteSessionChannel.h
@@ -3,7 +3,7 @@
 #pragma once
 
 #include "CoreMinimal.h"
-#include "RemoteSessionRole.h"
+#include "RemoteSession/RemoteSessionRole.h"
 
 class FBackChannelOSCConnection;
 
diff --git a/Engine/Plugins/Experimental/RemoteSession/Source/RemoteSessionEditor/Private/Widgets/SRemoteSessionStream.cpp b/Engine/Plugins/Experimental/RemoteSession/Source/RemoteSessionEditor/Private/Widgets/SRemoteSessionStream.cpp
--- a/Engine/Plugins/Experimental/RemoteSession/Source/RemoteSessionEditor/Private/Widgets/SRemoteSessionStream.cpp
+++ b/Engine/Plugins/Experimental/RemoteSession/Source/RemoteSessionEditor/Private/Widgets/SRemoteSessionStream.cpp
@@ -34,6 +34,7 @@
 
 #include "Channels/RemoteSessionImageChannel.h"
 #include "Channels/RemoteSessionInputChannel.h"
+#include "RemoteSession/RemoteSession.h"
 #include "ImageProviders/RemoteSessionMediaOutput.h"
 #include "RemoteSessionEditorStyle.h"
 
