Cleaned up the rendering code for drawing the remote AR camera to the scene

diff --git a/Engine/Plugins/Experimental/RemoteSession/Source/RemoteSession/Private/Channels/RemoteSessionARCameraChannel.cpp b/Engine/Plugins/Experimental/RemoteSession/Source/RemoteSession/Private/Channels/RemoteSessionARCameraChannel.cpp
--- a/Engine/Plugins/Experimental/RemoteSession/Source/RemoteSession/Private/Channels/RemoteSessionARCameraChannel.cpp
+++ b/Engine/Plugins/Experimental/RemoteSession/Source/RemoteSession/Private/Channels/RemoteSessionARCameraChannel.cpp
@@ -289,7 +289,6 @@ void FARCameraSceneViewExtension::PostRenderViewFamily_RenderThread(FRHICommandL
 
 void FARCameraSceneViewExtension::RenderARCamera_RenderThread(FRHICommandListImmediate& RHICmdList, const FSceneView& InView)
 {
-#if 0 // Removing until I can figure out the linking problem on PC
 	const auto FeatureLevel = InView.GetFeatureLevel();
 	IRendererModule& RendererModule = GetRendererModule();
 
@@ -306,21 +305,24 @@ void FARCameraSceneViewExtension::RenderARCamera_RenderThread(FRHICommandListImm
 	GraphicsPSOInit.BoundShaderState.VertexDeclarationRHI = GFilterVertexDeclaration.VertexDeclarationRHI;
 
 	const bool bIsMobileRenderer = FeatureLevel <= ERHIFeatureLevel::ES3_1;
+	FRemoteSessionARCameraVS<true>* VertexShaderPtr1 = nullptr;
+	FRemoteSessionARCameraVS<false>* VertexShaderPtr2 = nullptr;
+	FRemoteSessionARCameraPS* PixelShaderPtr = MaterialShaderMap->GetShader<FRemoteSessionARCameraPS>();
 	FMaterialShader* VertexShader = nullptr;
-	FMaterialShader* PixelShader = nullptr;
+	FMaterialShader* PixelShader = PixelShaderPtr;
 	if (bIsMobileRenderer)
 	{
-		VertexShader = MaterialShaderMap->GetShader<FRemoteSessionARCameraVS<true>>();
-		PixelShader = MaterialShaderMap->GetShader<FRemoteSessionARCameraPS>();
-		GraphicsPSOInit.BoundShaderState.VertexShaderRHI = GETSAFERHISHADER_VERTEX(reinterpret_cast<FRemoteSessionARCameraVS<true>*>(VertexShader));
-		GraphicsPSOInit.BoundShaderState.PixelShaderRHI = GETSAFERHISHADER_PIXEL(reinterpret_cast<FRemoteSessionARCameraPS*>(PixelShader));
+		VertexShaderPtr1 = MaterialShaderMap->GetShader<FRemoteSessionARCameraVS<true>>();
+		VertexShader = VertexShaderPtr1;
+		GraphicsPSOInit.BoundShaderState.VertexShaderRHI = GETSAFERHISHADER_VERTEX(VertexShader);
+		GraphicsPSOInit.BoundShaderState.PixelShaderRHI = GETSAFERHISHADER_PIXEL(PixelShader);
 	}
 	else
 	{
-		VertexShader = MaterialShaderMap->GetShader<FRemoteSessionARCameraVS<false>>();
-		PixelShader = MaterialShaderMap->GetShader<FRemoteSessionARCameraPS>();
-		GraphicsPSOInit.BoundShaderState.VertexShaderRHI = GETSAFERHISHADER_VERTEX(reinterpret_cast<FRemoteSessionARCameraVS<false>*>(VertexShader));
-		GraphicsPSOInit.BoundShaderState.PixelShaderRHI = GETSAFERHISHADER_PIXEL(reinterpret_cast<FRemoteSessionARCameraPS*>(PixelShader));
+		VertexShaderPtr2 = MaterialShaderMap->GetShader<FRemoteSessionARCameraVS<false>>();
+		VertexShader = VertexShaderPtr2;
+		GraphicsPSOInit.BoundShaderState.VertexShaderRHI = GETSAFERHISHADER_VERTEX(VertexShader);
+		GraphicsPSOInit.BoundShaderState.PixelShaderRHI = GETSAFERHISHADER_PIXEL(PixelShader);
 	}
 
 	SetGraphicsPipelineState(RHICmdList, GraphicsPSOInit);
@@ -335,17 +337,15 @@ void FARCameraSceneViewExtension::RenderARCamera_RenderThread(FRHICommandListImm
 
 	if (bIsMobileRenderer)
 	{
-		FRemoteSessionARCameraVS<true>* const VertexShaderPtr = reinterpret_cast<FRemoteSessionARCameraVS<true>*>(VertexShader);
-		SetUniformBufferParameterImmediate(RHICmdList, VertexShaderPtr->GetVertexShader(), VertexShaderPtr->GetUniformBufferParameter<FDrawRectangleParameters>(), Parameters);
-		VertexShaderPtr->SetParameters(RHICmdList, InView);
-		reinterpret_cast<FRemoteSessionARCameraPS*>(PixelShader)->SetParameters(RHICmdList, InView, PPMaterial->GetRenderProxy(false));
+		SetUniformBufferParameterImmediate(RHICmdList, VertexShaderPtr1->GetVertexShader(), VertexShaderPtr1->GetUniformBufferParameter<FDrawRectangleParameters>(), Parameters);
+		VertexShaderPtr1->SetParameters(RHICmdList, InView);
+		PixelShaderPtr->SetParameters(RHICmdList, InView, PPMaterial->GetRenderProxy(false));
 	}
 	else
 	{
-		FRemoteSessionARCameraVS<false>* const VertexShaderPtr = reinterpret_cast<FRemoteSessionARCameraVS<false>*>(VertexShader);
-		SetUniformBufferParameterImmediate(RHICmdList, VertexShaderPtr->GetVertexShader(), VertexShaderPtr->GetUniformBufferParameter<FDrawRectangleParameters>(), Parameters);
-		VertexShaderPtr->SetParameters(RHICmdList, InView);
-		reinterpret_cast<FRemoteSessionARCameraPS*>(PixelShader)->SetParameters(RHICmdList, InView, PPMaterial->GetRenderProxy(false));
+		SetUniformBufferParameterImmediate(RHICmdList, VertexShaderPtr2->GetVertexShader(), VertexShaderPtr2->GetUniformBufferParameter<FDrawRectangleParameters>(), Parameters);
+		VertexShaderPtr2->SetParameters(RHICmdList, InView);
+		PixelShaderPtr->SetParameters(RHICmdList, InView, PPMaterial->GetRenderProxy(false));
 	}
 
 	if (VertexBufferRHI && IndexBufferRHI.IsValid())
@@ -353,7 +353,6 @@ void FARCameraSceneViewExtension::RenderARCamera_RenderThread(FRHICommandListImm
 		RHICmdList.SetStreamSource(0, VertexBufferRHI, 0);
 		RHICmdList.DrawIndexedPrimitive(
 				IndexBufferRHI,
-				PT_TriangleList,
 				/*BaseVertexIndex=*/ 0,
 				/*MinIndex=*/ 0,
 				/*NumVertices=*/ 4,
@@ -362,7 +361,6 @@ void FARCameraSceneViewExtension::RenderARCamera_RenderThread(FRHICommandListImm
 				/*NumInstances=*/ 1
 			);
 	}
-#endif
 }
 
 bool FARCameraSceneViewExtension::IsActiveThisFrame(FViewport* InViewport) const
