Fixed a typo in a struct name

diff --git a/Engine/Plugins/Experimental/RemoteSession/Source/RemoteSession/Channels/RemoteSessionARCameraChannel.h b/Engine/Plugins/Experimental/RemoteSession/Source/RemoteSession/Channels/RemoteSessionARCameraChannel.h
--- a/Engine/Plugins/Experimental/RemoteSession/Source/RemoteSession/Channels/RemoteSessionARCameraChannel.h
+++ b/Engine/Plugins/Experimental/RemoteSession/Source/RemoteSession/Channels/RemoteSessionARCameraChannel.h
@@ -60,9 +60,9 @@ private:
 		TSharedPtr<FAppleImageUtilsConversionTaskBase, ESPMode::ThreadSafe> AsyncTask;
 	};
 
-	struct FDeompressedImage
+	struct FDecompressedImage
 	{
-		FDeompressedImage() : Width(0), Height(0) {}
+		FDecompressedImage() : Width(0), Height(0) {}
 
 		int32 Width;
 		int32 Height;
@@ -75,7 +75,7 @@ private:
 	TArray<TSharedPtr<FCompressionTask, ESPMode::ThreadSafe>> CompressionQueue;
 
 	/** Images that have been received and are available for rendering */
-	TArray<TSharedPtr<FDeompressedImage, ESPMode::ThreadSafe>> DecompressionQueue;
+	TArray<TSharedPtr<FDecompressedImage, ESPMode::ThreadSafe>> DecompressionQueue;
 	FThreadSafeCounter DecompressionTaskCount;
 	FCriticalSection DecompressionQueueLock;
 
diff --git a/Engine/Plugins/Experimental/RemoteSession/Source/RemoteSession/Private/Channels/RemoteSessionARCameraChannel.cpp b/Engine/Plugins/Experimental/RemoteSession/Source/RemoteSession/Private/Channels/RemoteSessionARCameraChannel.cpp
--- a/Engine/Plugins/Experimental/RemoteSession/Source/RemoteSession/Private/Channels/RemoteSessionARCameraChannel.cpp
+++ b/Engine/Plugins/Experimental/RemoteSession/Source/RemoteSession/Private/Channels/RemoteSessionARCameraChannel.cpp
@@ -549,7 +549,7 @@ void FRemoteSessionARCameraChannel::ReceiveARCameraImage(FBackChannelOSCMessage&
 	}
 	DecompressionTaskCount.Increment();
 
-	TSharedPtr<FDeompressedImage, ESPMode::ThreadSafe> DecompressedImage = MakeShareable(new FDeompressedImage());
+	TSharedPtr<FDecompressedImage, ESPMode::ThreadSafe> DecompressedImage = MakeShareable(new FDecompressedImage());
 	Message << DecompressedImage->Width;
 	Message << DecompressedImage->Height;
 	Message << DecompressedImage->ImageData;
@@ -577,7 +577,7 @@ void FRemoteSessionARCameraChannel::ReceiveARCameraImage(FBackChannelOSCMessage&
 
 void FRemoteSessionARCameraChannel::UpdateRenderingTexture()
 {
-	TSharedPtr<FDeompressedImage, ESPMode::ThreadSafe> DecompressedImage;
+	TSharedPtr<FDecompressedImage, ESPMode::ThreadSafe> DecompressedImage;
 	{
 		FScopeLock sl(&DecompressionQueueLock);
 		if (DecompressionQueue.Num())
