Added some asserts to make static analysis happy

diff --git a/Engine/Plugins/Experimental/RemoteSession/Source/RemoteSession/Private/Channels/RemoteSessionARCameraChannel.cpp b/Engine/Plugins/Experimental/RemoteSession/Source/RemoteSession/Private/Channels/RemoteSessionARCameraChannel.cpp
--- a/Engine/Plugins/Experimental/RemoteSession/Source/RemoteSession/Private/Channels/RemoteSessionARCameraChannel.cpp
+++ b/Engine/Plugins/Experimental/RemoteSession/Source/RemoteSession/Private/Channels/RemoteSessionARCameraChannel.cpp
@@ -314,6 +314,7 @@ void FARCameraSceneViewExtension::RenderARCamera_RenderThread(FRHICommandListImm
 	{
 		VertexShaderPtr1 = MaterialShaderMap->GetShader<FRemoteSessionARCameraVS<true>>();
 		VertexShader = VertexShaderPtr1;
+		check(PixelShader != nullptr && VertexShader != nullptr && VertexShaderPtr1 != nullptr);
 		GraphicsPSOInit.BoundShaderState.VertexShaderRHI = GETSAFERHISHADER_VERTEX(VertexShader);
 		GraphicsPSOInit.BoundShaderState.PixelShaderRHI = GETSAFERHISHADER_PIXEL(PixelShader);
 	}
@@ -321,9 +322,11 @@ void FARCameraSceneViewExtension::RenderARCamera_RenderThread(FRHICommandListImm
 	{
 		VertexShaderPtr2 = MaterialShaderMap->GetShader<FRemoteSessionARCameraVS<false>>();
 		VertexShader = VertexShaderPtr2;
+		check(PixelShader != nullptr && VertexShader != nullptr && VertexShaderPtr2 != nullptr);
 		GraphicsPSOInit.BoundShaderState.VertexShaderRHI = GETSAFERHISHADER_VERTEX(VertexShader);
 		GraphicsPSOInit.BoundShaderState.PixelShaderRHI = GETSAFERHISHADER_PIXEL(PixelShader);
 	}
+	check(PixelShader != nullptr && VertexShader != nullptr && (VertexShaderPtr1 != nullptr || VertexShaderPtr2 != nullptr));
 
 	SetGraphicsPipelineState(RHICmdList, GraphicsPSOInit);
 
