Removed some code from the remote session client that was overriding the receive buffer size that the config files had set

diff --git a/Engine/Plugins/Experimental/RemoteSession/Source/RemoteSession/Private/RemoteSessionClient.cpp b/Engine/Plugins/Experimental/RemoteSession/Source/RemoteSession/Private/RemoteSessionClient.cpp
--- a/Engine/Plugins/Experimental/RemoteSession/Source/RemoteSession/Private/RemoteSessionClient.cpp
+++ b/Engine/Plugins/Experimental/RemoteSession/Source/RemoteSession/Private/RemoteSessionClient.cpp
@@ -102,15 +102,11 @@ void FRemoteSessionClient::CheckConnection()
 
 	// success indicates that our check was successful, if our connection was successful then
 	// the delegate code is called
-	bool Success = Connection->WaitForConnection(0, [this](auto InConnection) {
-		int32 WantedSize = 4 * 1024 * 1024;
-		int32 ActualSize(0);
-
-		Connection->GetSocket()->SetReceiveBufferSize(WantedSize, ActualSize);
-
+	bool Success = Connection->WaitForConnection(0, [this](auto InConnection)
+	{
 		CreateOSCConnection(Connection.ToSharedRef());
 		
-		UE_LOG(LogRemoteSession, Log, TEXT("Connected to host at %s (ReceiveSize=%dkb)"), *HostAddress, ActualSize / 1024);
+		UE_LOG(LogRemoteSession, Log, TEXT("Connected to host at %s"), *HostAddress);
 
 		IsConnecting = false;
 
