Added checks for device support before enabling certain remote session channels
Simplified the shader setup for the remote ar image rendering
Added a check to see if a AR session is running before automatically starting a session for the AR remote channels

diff --git a/Engine/Plugins/Experimental/RemoteSession/Source/RemoteSession/Private/Channels/RemoteSessionARCameraChannel.cpp b/Engine/Plugins/Experimental/RemoteSession/Source/RemoteSession/Private/Channels/RemoteSessionARCameraChannel.cpp
--- a/Engine/Plugins/Experimental/RemoteSession/Source/RemoteSession/Private/Channels/RemoteSessionARCameraChannel.cpp
+++ b/Engine/Plugins/Experimental/RemoteSession/Source/RemoteSession/Private/Channels/RemoteSessionARCameraChannel.cpp
@@ -61,7 +61,6 @@ TAutoConsoleVariable<int32> CVarJPEGGpu(
 	ECVF_Default);
 
 /** Shaders to render our post process material */
-template <bool bIsMobileRenderer>
 class FRemoteSessionARCameraVS :
 	public FMaterialShader
 {
@@ -71,14 +70,7 @@ public:
 
 	static bool ShouldCompilePermutation(EShaderPlatform Platform, const FMaterial* Material)
 	{
-		if (bIsMobileRenderer)
-		{
-			return Material->GetMaterialDomain() == MD_PostProcess && IsMobilePlatform(Platform);
-		}
-		else
-		{
-			return Material->GetMaterialDomain() == MD_PostProcess && !IsMobilePlatform(Platform);
-		}
+		return Material->GetMaterialDomain() == MD_PostProcess && !IsMobilePlatform(Platform);
 	}
 
 	static void ModifyCompilationEnvironment(EShaderPlatform Platform, const class FMaterial* Material, FShaderCompilerEnvironment& OutEnvironment)
@@ -107,8 +99,7 @@ public:
 	}
 };
 
-IMPLEMENT_MATERIAL_SHADER_TYPE(template<>, FRemoteSessionARCameraVS<true>, TEXT("/Engine/Private/PostProcessMaterialShaders.usf"), TEXT("MainVS_ES2"), SF_Vertex);
-IMPLEMENT_MATERIAL_SHADER_TYPE(template<>, FRemoteSessionARCameraVS<false>, TEXT("/Engine/Private/PostProcessMaterialShaders.usf"), TEXT("MainVS"), SF_Vertex);
+IMPLEMENT_MATERIAL_SHADER_TYPE(, FRemoteSessionARCameraVS, TEXT("/Engine/Private/PostProcessMaterialShaders.usf"), TEXT("MainVS"), SF_Vertex);
 
 class FRemoteSessionARCameraPS :
 	public FMaterialShader
@@ -119,7 +110,7 @@ public:
 
 	static bool ShouldCompilePermutation(EShaderPlatform Platform, const FMaterial* Material)
 	{
-		return Material->GetMaterialDomain() == MD_PostProcess;
+		return Material->GetMaterialDomain() == MD_PostProcess && !IsMobilePlatform(Platform);
 	}
 
 	static void ModifyCompilationEnvironment(EShaderPlatform Platform, const class FMaterial* Material, FShaderCompilerEnvironment& OutEnvironment)
@@ -172,7 +163,7 @@ private:
 	FShaderResourceParameter PostprocessInputParameterSampler[ePId_Input_MAX];
 };
 
-IMPLEMENT_MATERIAL_SHADER_TYPE(, FRemoteSessionARCameraPS, TEXT("/Engine/Private/PostProcessPassThrough.usf"), TEXT("MainPS"), SF_Pixel);
+IMPLEMENT_MATERIAL_SHADER_TYPE(, FRemoteSessionARCameraPS, TEXT("/Engine/Private/PostProcessMaterialShaders.usf"), TEXT("MainPS_VideoOverlay"), SF_Pixel);
 
 class FARCameraSceneViewExtension :
 	public FSceneViewExtensionBase
@@ -289,7 +280,9 @@ void FARCameraSceneViewExtension::PostRenderViewFamily_RenderThread(FRHICommandL
 
 void FARCameraSceneViewExtension::RenderARCamera_RenderThread(FRHICommandListImmediate& RHICmdList, const FSceneView& InView)
 {
+#if PLATFORM_DESKTOP
 	const auto FeatureLevel = InView.GetFeatureLevel();
+
 	IRendererModule& RendererModule = GetRendererModule();
 
 	const FMaterial* const CameraMaterial = PPMaterial->GetRenderProxy(false)->GetMaterial(FeatureLevel);
@@ -304,29 +297,12 @@ void FARCameraSceneViewExtension::RenderARCamera_RenderThread(FRHICommandListImm
 	GraphicsPSOInit.PrimitiveType = PT_TriangleList;
 	GraphicsPSOInit.BoundShaderState.VertexDeclarationRHI = GFilterVertexDeclaration.VertexDeclarationRHI;
 
-	const bool bIsMobileRenderer = FeatureLevel <= ERHIFeatureLevel::ES3_1;
-	FRemoteSessionARCameraVS<true>* VertexShaderPtr1 = nullptr;
-	FRemoteSessionARCameraVS<false>* VertexShaderPtr2 = nullptr;
-	FRemoteSessionARCameraPS* PixelShaderPtr = MaterialShaderMap->GetShader<FRemoteSessionARCameraPS>();
-	FMaterialShader* VertexShader = nullptr;
-	FMaterialShader* PixelShader = PixelShaderPtr;
-	if (bIsMobileRenderer)
-	{
-		VertexShaderPtr1 = MaterialShaderMap->GetShader<FRemoteSessionARCameraVS<true>>();
-		VertexShader = VertexShaderPtr1;
-		check(PixelShader != nullptr && VertexShader != nullptr && VertexShaderPtr1 != nullptr);
-		GraphicsPSOInit.BoundShaderState.VertexShaderRHI = GETSAFERHISHADER_VERTEX(VertexShader);
-		GraphicsPSOInit.BoundShaderState.PixelShaderRHI = GETSAFERHISHADER_PIXEL(PixelShader);
-	}
-	else
-	{
-		VertexShaderPtr2 = MaterialShaderMap->GetShader<FRemoteSessionARCameraVS<false>>();
-		VertexShader = VertexShaderPtr2;
-		check(PixelShader != nullptr && VertexShader != nullptr && VertexShaderPtr2 != nullptr);
-		GraphicsPSOInit.BoundShaderState.VertexShaderRHI = GETSAFERHISHADER_VERTEX(VertexShader);
-		GraphicsPSOInit.BoundShaderState.PixelShaderRHI = GETSAFERHISHADER_PIXEL(PixelShader);
-	}
-	check(PixelShader != nullptr && VertexShader != nullptr && (VertexShaderPtr1 != nullptr || VertexShaderPtr2 != nullptr));
+	FRemoteSessionARCameraVS* VertexShader = MaterialShaderMap->GetShader<FRemoteSessionARCameraVS>();
+	FRemoteSessionARCameraPS* PixelShader = MaterialShaderMap->GetShader<FRemoteSessionARCameraPS>();
+	check(PixelShader != nullptr && VertexShader != nullptr);
+
+	GraphicsPSOInit.BoundShaderState.VertexShaderRHI = GETSAFERHISHADER_VERTEX(VertexShader);
+	GraphicsPSOInit.BoundShaderState.PixelShaderRHI = GETSAFERHISHADER_PIXEL(PixelShader);
 
 	SetGraphicsPipelineState(RHICmdList, GraphicsPSOInit);
 
@@ -338,18 +314,9 @@ void FARCameraSceneViewExtension::RenderARCamera_RenderThread(FRHICommandListImm
 			1.0f / ViewSize.X, 1.0f / ViewSize.Y,
 			1.0f, 1.0f);
 
-	if (bIsMobileRenderer)
-	{
-		SetUniformBufferParameterImmediate(RHICmdList, VertexShaderPtr1->GetVertexShader(), VertexShaderPtr1->GetUniformBufferParameter<FDrawRectangleParameters>(), Parameters);
-		VertexShaderPtr1->SetParameters(RHICmdList, InView);
-		PixelShaderPtr->SetParameters(RHICmdList, InView, PPMaterial->GetRenderProxy(false));
-	}
-	else
-	{
-		SetUniformBufferParameterImmediate(RHICmdList, VertexShaderPtr2->GetVertexShader(), VertexShaderPtr2->GetUniformBufferParameter<FDrawRectangleParameters>(), Parameters);
-		VertexShaderPtr2->SetParameters(RHICmdList, InView);
-		PixelShaderPtr->SetParameters(RHICmdList, InView, PPMaterial->GetRenderProxy(false));
-	}
+	SetUniformBufferParameterImmediate(RHICmdList, VertexShader->GetVertexShader(), VertexShader->GetUniformBufferParameter<FDrawRectangleParameters>(), Parameters);
+	VertexShader->SetParameters(RHICmdList, InView);
+	PixelShader->SetParameters(RHICmdList, InView, PPMaterial->GetRenderProxy(false));
 
 	if (VertexBufferRHI && IndexBufferRHI.IsValid())
 	{
@@ -364,11 +331,12 @@ void FARCameraSceneViewExtension::RenderARCamera_RenderThread(FRHICommandListImm
 				/*NumInstances=*/ 1
 			);
 	}
+#endif
 }
 
 bool FARCameraSceneViewExtension::IsActiveThisFrame(FViewport* InViewport) const
 {
-	return Channel.GetPostProcessMaterial() != nullptr;
+	return PLATFORM_DESKTOP && Channel.GetPostProcessMaterial() != nullptr;
 }
 
 static FName CameraImageParamName(TEXT("CameraImage"));
@@ -380,13 +348,13 @@ FRemoteSessionARCameraChannel::FRemoteSessionARCameraChannel(ERemoteSessionChann
 	, Connection(InConnection)
 	, Role(InRole)
 {
-	
-	static bool OnceTimeARInit = false;
-	if (!OnceTimeARInit && InRole == ERemoteSessionChannelMode::Write)
+	if (InRole == ERemoteSessionChannelMode::Write)
 	{
-		UARSessionConfig* Config = NewObject<UARSessionConfig>();
-		UARBlueprintLibrary::StartARSession(Config);
-		OnceTimeARInit = true;
+		if (UARBlueprintLibrary::GetARSessionStatus().Status != EARSessionStatus::Running)
+		{
+			UARSessionConfig* Config = NewObject<UARSessionConfig>();
+			UARBlueprintLibrary::StartARSession(Config);
+		}
 	}
 	
 	RenderingTextures[0] = nullptr;
diff --git a/Engine/Plugins/Experimental/RemoteSession/Source/RemoteSession/Private/Channels/RemoteSessionXRTrackingChannel.cpp b/Engine/Plugins/Experimental/RemoteSession/Source/RemoteSession/Private/Channels/RemoteSessionXRTrackingChannel.cpp
--- a/Engine/Plugins/Experimental/RemoteSession/Source/RemoteSession/Private/Channels/RemoteSessionXRTrackingChannel.cpp
+++ b/Engine/Plugins/Experimental/RemoteSession/Source/RemoteSession/Private/Channels/RemoteSessionXRTrackingChannel.cpp
@@ -57,12 +57,12 @@ FRemoteSessionXRTrackingChannel::FRemoteSessionXRTrackingChannel(ERemoteSessionC
 	}
     else
     {
-        // Initialize AR if desired - todo, does this need to check device caps?
 #if PLATFORM_IOS
-		// Workaround - we don't want to set bSupportAR in our project as it prevents us running on old devices, but this needs
-		// to be true before we try to init ARKit stuff
-		UARSessionConfig* Config = NewObject<UARSessionConfig>();
-		UARBlueprintLibrary::StartARSession(Config);
+		if (UARBlueprintLibrary::GetARSessionStatus().Status != EARSessionStatus::Running)
+		{
+			UARSessionConfig* Config = NewObject<UARSessionConfig>();
+			UARBlueprintLibrary::StartARSession(Config);
+		}
 #endif
     }
 }
diff --git a/Engine/Plugins/Experimental/RemoteSession/Source/RemoteSession/Private/RemoteSessionRole.cpp b/Engine/Plugins/Experimental/RemoteSession/Source/RemoteSession/Private/RemoteSessionRole.cpp
--- a/Engine/Plugins/Experimental/RemoteSession/Source/RemoteSession/Private/RemoteSessionRole.cpp
+++ b/Engine/Plugins/Experimental/RemoteSession/Source/RemoteSession/Private/RemoteSessionRole.cpp
@@ -10,6 +10,7 @@
 #include "Channels/RemoteSessionFrameBufferChannel.h"
 #include "Async/Async.h"
 #include "GeneralProjectSettings.h"
+#include "ARBlueprintLibrary.h"
 
 DEFINE_LOG_CATEGORY(LogRemoteSession);
 
@@ -222,19 +223,27 @@ void FRemoteSessionRole::CreateChannel(const FString& InChannelName, ERemoteSess
 	}
 	else if (InChannelName == FRemoteSessionXRTrackingChannel::StaticType())
 	{
-		NewChannel = MakeShareable(new FRemoteSessionXRTrackingChannel(InMode, OSCConnection));
+		bool IsSupported = (InMode == ERemoteSessionChannelMode::Read) || UARBlueprintLibrary::IsSessionTypeSupported(EARSessionType::World);
+		if (IsSupported)
+		{
+			NewChannel = MakeShareable(new FRemoteSessionXRTrackingChannel(InMode, OSCConnection));
+		}
+		else
+		{
+			UE_LOG(LogRemoteSession, Warning, TEXT("FRemoteSessionXRTrackingChannel does not support sending on this platform/device"));
+		}
 	}
 	else if (InChannelName == FRemoteSessionARCameraChannel::StaticType())
 	{
 		// Client side sending only works on iOS with Android coming in the future
-		bool IsSupported = (InMode == ERemoteSessionChannelMode::Read) || PLATFORM_IOS;
+		bool IsSupported = (InMode == ERemoteSessionChannelMode::Read) || (PLATFORM_IOS && UARBlueprintLibrary::IsSessionTypeSupported(EARSessionType::World));
 		if (IsSupported)
 		{
 			NewChannel = MakeShareable(new FRemoteSessionARCameraChannel(InMode, OSCConnection));
 		}
 		else
 		{
-			UE_LOG(LogRemoteSession, Warning, TEXT("FRemoteSessionARCameraChannel does not support sending on this platform"));
+			UE_LOG(LogRemoteSession, Warning, TEXT("FRemoteSessionARCameraChannel does not support sending on this platform/device"));
 		}
 	}
 	
