Fixed mismatched vertex/pixel shader combo

diff --git a/Engine/Plugins/Experimental/RemoteSession/Source/RemoteSession/Private/Channels/RemoteSessionARCameraChannel.cpp b/Engine/Plugins/Experimental/RemoteSession/Source/RemoteSession/Private/Channels/RemoteSessionARCameraChannel.cpp
--- a/Engine/Plugins/Experimental/RemoteSession/Source/RemoteSession/Private/Channels/RemoteSessionARCameraChannel.cpp
+++ b/Engine/Plugins/Experimental/RemoteSession/Source/RemoteSession/Private/Channels/RemoteSessionARCameraChannel.cpp
@@ -99,7 +99,7 @@ public:
 	}
 };
 
-IMPLEMENT_MATERIAL_SHADER_TYPE(, FRemoteSessionARCameraVS, TEXT("/Engine/Private/PostProcessMaterialShaders.usf"), TEXT("MainVS"), SF_Vertex);
+IMPLEMENT_MATERIAL_SHADER_TYPE(, FRemoteSessionARCameraVS, TEXT("/Engine/Private/PostProcessMaterialShaders.usf"), TEXT("MainVS_VideoOverlay"), SF_Vertex);
 
 class FRemoteSessionARCameraPS :
 	public FMaterialShader
