Changed how the last image in the decompression queue was selected for the remote ar camera image

diff --git a/Engine/Plugins/Experimental/RemoteSession/Source/RemoteSession/Private/Channels/RemoteSessionARCameraChannel.cpp b/Engine/Plugins/Experimental/RemoteSession/Source/RemoteSession/Private/Channels/RemoteSessionARCameraChannel.cpp
--- a/Engine/Plugins/Experimental/RemoteSession/Source/RemoteSession/Private/Channels/RemoteSessionARCameraChannel.cpp
+++ b/Engine/Plugins/Experimental/RemoteSession/Source/RemoteSession/Private/Channels/RemoteSessionARCameraChannel.cpp
@@ -461,20 +461,16 @@ void FRemoteSessionARCameraChannel::SendARCameraImage()
 	}
 
 	TSharedPtr<FCompressionTask, ESPMode::ThreadSafe> CompressionTask;
+	if (CompressionQueue.Num() > 0)
 	{
 		int32 CompleteIndex = -1;
-		bool bFound = false;
-		for (int32 Index = 0; Index < CompressionQueue.Num() && bFound; Index++)
+		// Find the latest task that has completed
+		for (int32 Index = CompressionQueue.Num() - 1; Index >= 0; Index--)
 		{
-			// Find the latest task that has completed
 			if (CompressionQueue[Index]->AsyncTask->IsDone())
 			{
 				CompleteIndex = Index;
-			}
-			// If this task isn't done, but the one before it is, then we're done
-			else if (CompleteIndex > -1 || Index == 0)
-			{
-				bFound = true;
+				break;
 			}
 		}
 		if (CompleteIndex > -1)
