Fixed using the wrong material to draw causing it to miss the texture parameter
Changed the outgoing compression to use a single task rather than a queue to remove pressure from the GPU

diff --git a/Engine/Plugins/Experimental/RemoteSession/Source/RemoteSession/Channels/RemoteSessionARCameraChannel.h b/Engine/Plugins/Experimental/RemoteSession/Source/RemoteSession/Channels/RemoteSessionARCameraChannel.h
--- a/Engine/Plugins/Experimental/RemoteSession/Source/RemoteSession/Channels/RemoteSessionARCameraChannel.h
+++ b/Engine/Plugins/Experimental/RemoteSession/Source/RemoteSession/Channels/RemoteSessionARCameraChannel.h
@@ -69,10 +69,8 @@ private:
 		TArray<uint8> ImageData;
 	};
 
-	/** The timestamp of the last queued image to prevent from sending the same image repeatedly */
-	float LastQueuedTimestamp;
-	/** List of async image compression tasks in flight */
-	TArray<TSharedPtr<FCompressionTask, ESPMode::ThreadSafe>> CompressionQueue;
+	/** Only compress one frame at a time to prevent flooding the GPU */
+	TSharedPtr<FCompressionTask, ESPMode::ThreadSafe> CompressionTask;
 
 	/** Images that have been received and are available for rendering */
 	TArray<TSharedPtr<FDecompressedImage, ESPMode::ThreadSafe>> DecompressionQueue;
@@ -85,6 +83,7 @@ private:
 	UMaterialInstanceDynamic* MaterialInstanceDynamic;
 	/** Textures to use when rendering */
 	UTexture2D* RenderingTextures[2];
+	UTexture2D* LastSetTexture;
 	FThreadSafeCounter RenderingTextureIndex;
 	FThreadSafeCounter RenderingTexturesUpdateCount[2];
 
diff --git a/Engine/Plugins/Experimental/RemoteSession/Source/RemoteSession/Private/Channels/RemoteSessionARCameraChannel.cpp b/Engine/Plugins/Experimental/RemoteSession/Source/RemoteSession/Private/Channels/RemoteSessionARCameraChannel.cpp
--- a/Engine/Plugins/Experimental/RemoteSession/Source/RemoteSession/Private/Channels/RemoteSessionARCameraChannel.cpp
+++ b/Engine/Plugins/Experimental/RemoteSession/Source/RemoteSession/Private/Channels/RemoteSessionARCameraChannel.cpp
@@ -343,7 +343,6 @@ static FName CameraImageParamName(TEXT("CameraImage"));
 
 FRemoteSessionARCameraChannel::FRemoteSessionARCameraChannel(ERemoteSessionChannelMode InRole, TSharedPtr<FBackChannelOSCConnection, ESPMode::ThreadSafe> InConnection)
 	: IRemoteSessionChannel(InRole, InConnection)
-	, LastQueuedTimestamp(0.f)
 	, RenderingTextureIndex(0)
 	, Connection(InConnection)
 	, Role(InRole)
@@ -359,6 +358,7 @@ FRemoteSessionARCameraChannel::FRemoteSessionARCameraChannel(ERemoteSessionChann
 	
 	RenderingTextures[0] = nullptr;
 	RenderingTextures[1] = nullptr;
+	LastSetTexture = nullptr;
 	PPMaterial = LoadObject<UMaterialInterface>(nullptr, TEXT("/RemoteSession/ARCameraPostProcess.ARCameraPostProcess"));
 	MaterialInstanceDynamic = UMaterialInstanceDynamic::Create(PPMaterial, GetTransientPackage());
 	if (MaterialInstanceDynamic != nullptr)
@@ -402,8 +402,10 @@ void FRemoteSessionARCameraChannel::Tick(const float InDeltaTime)
 #if PLATFORM_IOS
 	if (Role == ERemoteSessionChannelMode::Write)
 	{
-		QueueARCameraImage();
+		// Always send a complete compression task to make room for another
 		SendARCameraImage();
+		// Queue a compression task if we don't have one outstanding
+		QueueARCameraImage();
 	}
 #endif
 	if (Role == ERemoteSessionChannelMode::Read)
@@ -413,8 +415,13 @@ void FRemoteSessionARCameraChannel::Tick(const float InDeltaTime)
 		{
 			if (UTexture2D* NextTexture = RenderingTextures[RenderingTextureIndex.GetValue()])
 			{
-				// Update the texture to the current one
-				MaterialInstanceDynamic->SetTextureParameterValue(CameraImageParamName, NextTexture);
+				// Only update the material when then texture changes
+				if (LastSetTexture != NextTexture)
+				{
+					LastSetTexture = NextTexture;
+					// Update the texture to the current one
+					MaterialInstanceDynamic->SetTextureParameterValue(CameraImageParamName, NextTexture);
+				}
 			}
 		}
 	}
@@ -429,21 +436,19 @@ void FRemoteSessionARCameraChannel::QueueARCameraImage()
 		return;
 	}
 
+	// Don't queue multiple compression tasks at the same time or we get a well of despair on the GPU
+	if (CompressionTask.IsValid())
+	{
+		return;
+	}
+
 	UARTextureCameraImage* CameraImage = UARBlueprintLibrary::GetCameraImage();
 	if (CameraImage != nullptr)
     {
-        if (CameraImage->Timestamp > LastQueuedTimestamp)
-        {
-            TSharedPtr<FCompressionTask, ESPMode::ThreadSafe> CompressionTask = MakeShareable(new FCompressionTask());
-            CompressionTask->Width = CameraImage->Size.X;
-            CompressionTask->Height = CameraImage->Size.Y;
-            CompressionTask->AsyncTask = IAppleImageUtilsPlugin::Get().ConvertToJPEG(CameraImage, CVarJPEGQuality.GetValueOnGameThread(), !!CVarJPEGColor.GetValueOnGameThread(), !!CVarJPEGGpu.GetValueOnGameThread());
-            if (CompressionTask->AsyncTask.IsValid())
-            {
-                LastQueuedTimestamp = CameraImage->Timestamp;
-                CompressionQueue.Add(CompressionTask);
-            }
-        }
+		CompressionTask = MakeShareable(new FCompressionTask());
+		CompressionTask->Width = CameraImage->Size.X;
+		CompressionTask->Height = CameraImage->Size.Y;
+		CompressionTask->AsyncTask = IAppleImageUtilsPlugin::Get().ConvertToJPEG(CameraImage, CVarJPEGQuality.GetValueOnGameThread(), !!CVarJPEGColor.GetValueOnGameThread(), !!CVarJPEGGpu.GetValueOnGameThread());
     }
     else
     {
@@ -460,45 +465,34 @@ void FRemoteSessionARCameraChannel::SendARCameraImage()
 		return;
 	}
 
-	TSharedPtr<FCompressionTask, ESPMode::ThreadSafe> CompressionTask;
-	if (CompressionQueue.Num() > 0)
+	// Bail if there's nothing to do yet
+	if (!CompressionTask.IsValid() ||
+		(CompressionTask.IsValid() && !CompressionTask->AsyncTask->IsDone()))
 	{
-		int32 CompleteIndex = -1;
-		// Find the latest task that has completed
-		for (int32 Index = CompressionQueue.Num() - 1; Index >= 0; Index--)
-		{
-			if (CompressionQueue[Index]->AsyncTask->IsDone())
-			{
-				CompleteIndex = Index;
-				break;
-			}
-		}
-		if (CompleteIndex > -1)
-		{
-			// Grab the latest completed one
-			CompressionTask = CompressionQueue[CompleteIndex];
-			// And clear out all of the tasks between 0 and CompleteIndex
-			CompressionQueue.RemoveAt(0, CompleteIndex + 1);
-		}
+		return;
 	}
 
-	if (CompressionTask.IsValid() && !CompressionTask->AsyncTask->HadError())
+	if (!CompressionTask->AsyncTask->HadError())
 	{
-		AsyncTask(ENamedThreads::AnyBackgroundThreadNormalTask, [this, CompressionTask]()
+		// Copy the task so we can start GPU compressing on another one
+		TSharedPtr<FCompressionTask, ESPMode::ThreadSafe> SendCompressionTask = CompressionTask;
+		AsyncTask(ENamedThreads::AnyBackgroundThreadNormalTask, [this, SendCompressionTask]()
 		{
 			FBackChannelOSCMessage Msg(CAMERA_MESSAGE_ADDRESS);
-			Msg.Write(CompressionTask->Width);
-			Msg.Write(CompressionTask->Height);
-			Msg.Write(CompressionTask->AsyncTask->GetData());
+			Msg.Write(SendCompressionTask->Width);
+			Msg.Write(SendCompressionTask->Height);
+			Msg.Write(SendCompressionTask->AsyncTask->GetData());
 
 			Connection->SendPacket(Msg);
 		});
 	}
+	// Release this task so we can queue another one
+	CompressionTask.Reset();
 }
 
 UMaterialInterface* FRemoteSessionARCameraChannel::GetPostProcessMaterial() const
 {
-	return PPMaterial;
+	return MaterialInstanceDynamic;
 }
 
 void FRemoteSessionARCameraChannel::ReceiveARCameraImage(FBackChannelOSCMessage& Message, FBackChannelOSCDispatch& Dispatch)
