Turned sRGB off for the camera image that comes from the remote AR device

diff --git a/Engine/Plugins/Experimental/RemoteSession/Source/RemoteSession/Private/Channels/RemoteSessionARCameraChannel.cpp b/Engine/Plugins/Experimental/RemoteSession/Source/RemoteSession/Private/Channels/RemoteSessionARCameraChannel.cpp
--- a/Engine/Plugins/Experimental/RemoteSession/Source/RemoteSession/Private/Channels/RemoteSessionARCameraChannel.cpp
+++ b/Engine/Plugins/Experimental/RemoteSession/Source/RemoteSession/Private/Channels/RemoteSessionARCameraChannel.cpp
@@ -562,6 +562,7 @@ void FRemoteSessionARCameraChannel::UpdateRenderingTexture()
 		if (RenderingTextures[NextImage] == nullptr || DecompressedImage->Width != RenderingTextures[NextImage]->GetSizeX() || DecompressedImage->Height != RenderingTextures[NextImage]->GetSizeY())
 		{
 			RenderingTextures[NextImage] = UTexture2D::CreateTransient(DecompressedImage->Width, DecompressedImage->Height);
+			RenderingTextures[NextImage]->SRGB = 0;
 			RenderingTextures[NextImage]->UpdateResource();
 		}
 
