Added a configuration for the frame grabber size when sending desktop images to a client

diff --git a/Engine/Plugins/Experimental/RemoteSession/Source/RemoteSession/Private/Channels/RemoteSessionFrameBufferChannel.cpp b/Engine/Plugins/Experimental/RemoteSession/Source/RemoteSession/Private/Channels/RemoteSessionFrameBufferChannel.cpp
--- a/Engine/Plugins/Experimental/RemoteSession/Source/RemoteSession/Private/Channels/RemoteSessionFrameBufferChannel.cpp
+++ b/Engine/Plugins/Experimental/RemoteSession/Source/RemoteSession/Private/Channels/RemoteSessionFrameBufferChannel.cpp
@@ -34,6 +34,17 @@ static FAutoConsoleVariableRef CVarQualityOverride(
 	TEXT("Sets quality (1-100)"),
 	ECVF_Default);
 
+static int32 FrameGrabberResX = 0;
+static FAutoConsoleVariableRef CVarResXOverride(
+	TEXT("remote.framegrabber.resx"), FrameGrabberResX,
+	TEXT("Sets the desired X resolution"),
+	ECVF_Default);
+
+static int32 FrameGrabberResY = 0;
+static FAutoConsoleVariableRef CVarResYOverride(
+	TEXT("remote.framegrabber.resy"), FrameGrabberResY,
+	TEXT("Sets the desired Y resolution"),
+	ECVF_Default);
 
 FRemoteSessionFrameBufferChannel::FRemoteSessionFrameBufferChannel(ERemoteSessionChannelMode InRole, TSharedPtr<FBackChannelOSCConnection, ESPMode::ThreadSafe> InConnection)
 	: IRemoteSessionChannel(InRole, InConnection)
@@ -141,7 +152,19 @@ void FRemoteSessionFrameBufferChannel::SetCaptureViewport(TSharedRef<FSceneViewp
 void FRemoteSessionFrameBufferChannel::CreateFrameGrabber(TSharedRef<FSceneViewport> Viewport)
 {
 	ReleaseFrameGrabber();
-	FrameGrabber = MakeShareable(new FFrameGrabber(Viewport, Viewport->GetSize()));
+
+	// For times when we want a specific resolution
+	FIntPoint FrameGrabberSize = SceneViewport->GetSize();
+	if (FrameGrabberResX > 0)
+	{
+		FrameGrabberSize.X = FrameGrabberResX;
+	}
+	if (FrameGrabberResY > 0)
+	{
+		FrameGrabberSize.Y = FrameGrabberResY;
+	}
+
+	FrameGrabber = MakeShareable(new FFrameGrabber(Viewport, FrameGrabberSize));
 	FrameGrabber->StartCapturingFrames();
 }
 
