Added support to the XR tracking channel to allow for a child class to specify an AR system

diff --git a/Engine/Plugins/Experimental/RemoteSession/Source/RemoteSession/Channels/RemoteSessionXRTrackingChannel.h b/Engine/Plugins/Experimental/RemoteSession/Source/RemoteSession/Channels/RemoteSessionXRTrackingChannel.h
--- a/Engine/Plugins/Experimental/RemoteSession/Source/RemoteSession/Channels/RemoteSessionXRTrackingChannel.h
+++ b/Engine/Plugins/Experimental/RemoteSession/Source/RemoteSession/Channels/RemoteSessionXRTrackingChannel.h
@@ -7,13 +7,14 @@
 
 class FBackChannelOSCMessage;
 class FBackChannelOSCDispatch;
+class IARSystemSupport;
 
 class REMOTESESSION_API FXRTrackingProxy :
 	public FXRTrackingSystemBase
 {
 public:
-	FXRTrackingProxy()
-		: FXRTrackingSystemBase(nullptr)
+	FXRTrackingProxy(IARSystemSupport* InARSystemSupport)
+		: FXRTrackingSystemBase(InARSystemSupport)
 	{}
 		
 	virtual bool IsTracking(int32 DeviceId) override { return true; }
@@ -34,7 +35,7 @@ public:
 
 	FRemoteSessionXRTrackingChannel(ERemoteSessionChannelMode InRole, TSharedPtr<FBackChannelOSCConnection, ESPMode::ThreadSafe> InConnection);
 
-	~FRemoteSessionXRTrackingChannel();
+	virtual ~FRemoteSessionXRTrackingChannel();
 
 	virtual void Tick(const float InDeltaTime) override;
 
@@ -50,6 +51,9 @@ public:
 	/* End IRemoteSessionChannel implementation */
 
 protected:
+	/** Only to be called by child classes */
+	FRemoteSessionXRTrackingChannel(ERemoteSessionChannelMode InRole, TSharedPtr<FBackChannelOSCConnection, ESPMode::ThreadSafe> InConnection, IARSystemSupport* ARSystemSupport);
+
 	TSharedPtr<FBackChannelOSCConnection, ESPMode::ThreadSafe> Connection;
 	ERemoteSessionChannelMode Role;
 
@@ -58,6 +62,13 @@ protected:
 	/** Used to set the values from the remote client as the XRTracking's pose */
 	TSharedPtr<FXRTrackingProxy, ESPMode::ThreadSafe> ProxyXRSystem;
 
+	/** Used by a child class to pass in a AR system to use as part of the proxy */
+	IARSystemSupport* ARSystemSupport;
+
+private:
 	/** So we can manage callback lifetimes properly */
 	FDelegateHandle MessageCallbackHandle;
+
+	/** Used to finish construction of the class. Should be called from within the ctors */
+	void Init();
 };
diff --git a/Engine/Plugins/Experimental/RemoteSession/Source/RemoteSession/Private/Channels/RemoteSessionXRTrackingChannel.cpp b/Engine/Plugins/Experimental/RemoteSession/Source/RemoteSession/Private/Channels/RemoteSessionXRTrackingChannel.cpp
--- a/Engine/Plugins/Experimental/RemoteSession/Source/RemoteSession/Private/Channels/RemoteSessionXRTrackingChannel.cpp
+++ b/Engine/Plugins/Experimental/RemoteSession/Source/RemoteSession/Private/Channels/RemoteSessionXRTrackingChannel.cpp
@@ -41,30 +41,18 @@ FRemoteSessionXRTrackingChannel::FRemoteSessionXRTrackingChannel(ERemoteSessionC
 	: IRemoteSessionChannel(InRole, InConnection)
 	, Connection(InConnection)
 	, Role(InRole)
+	, ARSystemSupport(nullptr)
 {
-	// If we are sending, we grab the data from GEngine->XRSystem, otherwise we back the current one up for restore later
-	XRSystem = GEngine->XRSystem;
-	
-	if (Role == ERemoteSessionChannelMode::Read)
-	{
-		// Make the proxy and set GEngine->XRSystem to it
-		ProxyXRSystem = MakeShared<FXRTrackingProxy, ESPMode::ThreadSafe>();
-		GEngine->XRSystem = ProxyXRSystem;
-		
-		auto Delegate = FBackChannelDispatchDelegate::FDelegate::CreateRaw(this, &FRemoteSessionXRTrackingChannel::ReceiveXRTracking);
-		MessageCallbackHandle = Connection->AddMessageHandler(MESSAGE_ADDRESS, Delegate);
-		Connection->SetMessageOptions(MESSAGE_ADDRESS, 1);
-	}
-    else
-    {
-#if PLATFORM_IOS
-		if (UARBlueprintLibrary::GetARSessionStatus().Status != EARSessionStatus::Running)
-		{
-			UARSessionConfig* Config = NewObject<UARSessionConfig>();
-			UARBlueprintLibrary::StartARSession(Config);
-		}
-#endif
-    }
+	Init();
+}
+
+FRemoteSessionXRTrackingChannel::FRemoteSessionXRTrackingChannel(ERemoteSessionChannelMode InRole, TSharedPtr<FBackChannelOSCConnection, ESPMode::ThreadSafe> InConnection, IARSystemSupport* InARSystemSupport)
+	: IRemoteSessionChannel(InRole, InConnection)
+	, Connection(InConnection)
+	, Role(InRole)
+	, ARSystemSupport(InARSystemSupport)
+{
+	Init();
 }
 
 FRemoteSessionXRTrackingChannel::~FRemoteSessionXRTrackingChannel()
@@ -85,6 +73,33 @@ FRemoteSessionXRTrackingChannel::~FRemoteSessionXRTrackingChannel()
 	ProxyXRSystem = nullptr;
 }
 
+void FRemoteSessionXRTrackingChannel::Init()
+{
+	// If we are sending, we grab the data from GEngine->XRSystem, otherwise we back the current one up for restore later
+	XRSystem = GEngine->XRSystem;
+
+	if (Role == ERemoteSessionChannelMode::Read)
+	{
+		// Make the proxy and set GEngine->XRSystem to it
+		ProxyXRSystem = MakeShared<FXRTrackingProxy, ESPMode::ThreadSafe>(ARSystemSupport);
+		GEngine->XRSystem = ProxyXRSystem;
+
+		auto Delegate = FBackChannelDispatchDelegate::FDelegate::CreateRaw(this, &FRemoteSessionXRTrackingChannel::ReceiveXRTracking);
+		MessageCallbackHandle = Connection->AddMessageHandler(MESSAGE_ADDRESS, Delegate);
+		Connection->SetMessageOptions(MESSAGE_ADDRESS, 1);
+	}
+	else
+	{
+#if PLATFORM_IOS
+		if (UARBlueprintLibrary::GetARSessionStatus().Status != EARSessionStatus::Running)
+		{
+			UARSessionConfig* Config = NewObject<UARSessionConfig>();
+			UARBlueprintLibrary::StartARSession(Config);
+		}
+#endif
+	}
+}
+
 void FRemoteSessionXRTrackingChannel::Tick(const float InDeltaTime)
 {
 	// Inbound data gets handled as callbacks
