Added a channel to stream AR trackable items to remote desktops (WIP, lightly tested)

diff --git a/Engine/Plugins/Experimental/RemoteSession/Source/RemoteSession/Channels/RemoteSessionARSystemChannel.h b/Engine/Plugins/Experimental/RemoteSession/Source/RemoteSession/Channels/RemoteSessionARSystemChannel.h
new file mode 100644
--- /dev/null
+++ b/Engine/Plugins/Experimental/RemoteSession/Source/RemoteSession/Channels/RemoteSessionARSystemChannel.h
@@ -0,0 +1,139 @@
+// Copyright 1998-2019 Epic Games, Inc. All Rights Reserved.
+
+#pragma once
+
+#include "RemoteSessionXRTrackingChannel.h"
+#include "ARSystem.h"
+
+/**
+ * This class acts as if there is a AR system present on the desktop that is receiving data from a remote device
+ */
+class FARSystemProxy :
+	public IARSystemSupport,
+	public FGCObject,
+	public TSharedFromThis<FARSystemProxy, ESPMode::ThreadSafe>
+{
+public:
+	/** Only the factory method should construct this */
+	FARSystemProxy();
+	virtual ~FARSystemProxy();
+
+	static TSharedPtr<FARSystemProxy, ESPMode::ThreadSafe> Get();
+	static IARSystemSupport* GetARSystemPtr();
+	static void Destroy();
+
+	// ~IARSystemSupport
+	virtual EARTrackingQuality OnGetTrackingQuality() const override;
+	virtual FARSessionStatus OnGetARSessionStatus() const override;
+	virtual TArray<UARTrackedGeometry*> OnGetAllTrackedGeometries() const override;
+	virtual bool OnIsTrackingTypeSupported(EARSessionType SessionType) const override;
+	virtual EARWorldMappingState OnGetWorldMappingStatus() const override;
+	virtual TArray<FARVideoFormat> OnGetSupportedVideoFormats(EARSessionType SessionType) const override;
+	// Not supported methods
+	// @todo JoeG -- Look at supporting this
+	virtual UARLightEstimate* OnGetCurrentLightEstimate() const override { return nullptr; }
+	virtual UARTextureCameraImage* OnGetCameraImage() override { return nullptr; }
+	virtual UARTextureCameraDepth* OnGetCameraDepth() override { return nullptr; }
+	// End todo block
+	virtual void OnStartARSession(UARSessionConfig* Config) override {}
+	virtual void OnPauseARSession() override {}
+	virtual void OnStopARSession() override {}
+	virtual void OnSetAlignmentTransform(const FTransform& InAlignmentTransform) override {}
+	virtual bool OnAddManualEnvironmentCaptureProbe(FVector Location, FVector Extent) override { return false; }
+	virtual TSharedPtr<FARSaveWorldAsyncTask, ESPMode::ThreadSafe> OnSaveWorld() const override
+	{
+		return MakeShared<FARErrorSaveWorldAsyncTask, ESPMode::ThreadSafe>(TEXT("Not supported on proxy"));
+	}
+	virtual TSharedPtr<FARGetCandidateObjectAsyncTask, ESPMode::ThreadSafe> OnGetCandidateObject(FVector Location, FVector Extent) const override
+	{
+		return MakeShared<FARErrorGetCandidateObjectAsyncTask, ESPMode::ThreadSafe>(TEXT("Not supported on proxy"));
+	}
+	virtual bool OnAddRuntimeCandidateImage(UARSessionConfig* Config, UTexture2D* CandidateTexture, FString FriendlyName, float PhysicalWidth) override { return false; }
+	virtual TArray<FVector> OnGetPointCloud() const override { return TArray<FVector>(); }
+	virtual TArray<UARPin*> OnGetAllPins() const override { return TArray<UARPin*>(); }
+	virtual UARPin* OnPinComponent(USceneComponent* ComponentToPin, const FTransform& PinToWorldTransform, UARTrackedGeometry* TrackedGeometry = nullptr, const FName DebugName = NAME_None) { return nullptr; }
+	virtual void OnRemovePin(UARPin* PinToRemove) override {}
+	virtual TArray<FARTraceResult> OnLineTraceTrackedObjects( const FVector2D ScreenCoord, EARLineTraceChannels TraceChannels ) override { return TArray<FARTraceResult>(); }
+	virtual TArray<FARTraceResult> OnLineTraceTrackedObjects( const FVector Start, const FVector End, EARLineTraceChannels TraceChannels ) override { return TArray<FARTraceResult>(); }
+	virtual void* GetARSessionRawPointer() override { return nullptr; }
+	virtual void* GetGameThreadARFrameRawPointer() override { return nullptr; }
+	// ~IARSystemSupport
+
+	//~ FGCObject
+	virtual void AddReferencedObjects(FReferenceCollector& Collector) override;
+	//~ FGCObject
+
+	// Methods called by the message handlers
+	void SetSupportedVideoFormats(const TArray<FARVideoFormat>& InFormats);
+	void SetSessionConfig(UARSessionConfig* InConfig);
+	void AddTrackable(UARTrackedGeometry* Added);
+	UARTrackedGeometry* GetTrackable(FGuid UniqueId);
+	void NotifyUpdated(UARTrackedGeometry* Updated);
+	void RemoveTrackable(FGuid UniqueId);
+	// ~Methods called by the message handlers
+
+private:
+	static TSharedPtr<FARSystemProxy, ESPMode::ThreadSafe> FactoryInstance;
+
+	/** Used to return the session config object the remote is using */
+	UARSessionConfig* SessionConfig;
+	/** Map of unique ids to tracked geometries */
+	TMap<FGuid, UARTrackedGeometry*> TrackedGeometries;
+	/** The video formats that are supported */
+	TArray<FARVideoFormat> SupportedFormats;
+};
+
+/**
+ * Class used to deliver AR trackable items to the remote machine. Inherits from the XR tracking
+ * channel which handles sending all of the position & orientation updates. Do not use both this
+ * channel and the XR tracking at the same time, since they both hook GEngine->XRSystem and
+ * you will get inconsistent results
+ */
+class REMOTESESSION_API FRemoteSessionARSystemChannel :
+	public FRemoteSessionXRTrackingChannel
+{
+public:
+	FRemoteSessionARSystemChannel(ERemoteSessionChannelMode InRole, TSharedPtr<FBackChannelOSCConnection, ESPMode::ThreadSafe> InConnection);
+
+	virtual ~FRemoteSessionARSystemChannel();
+
+	/* Begin IRemoteSessionChannel implementation */
+	static const TCHAR* StaticType() { return TEXT("FRemoteSessionARSystemChannel"); }
+	virtual const TCHAR* GetType() const override { return StaticType(); }
+	/* End IRemoteSessionChannel implementation */
+
+	// Message handlers
+	void ReceiveARInit(FBackChannelOSCMessage& Message, FBackChannelOSCDispatch& Dispatch);
+	void ReceiveAddTrackable(FBackChannelOSCMessage& Message, FBackChannelOSCDispatch& Dispatch);
+	void ReceiveUpdateTrackable(FBackChannelOSCMessage& Message, FBackChannelOSCDispatch& Dispatch);
+	void ReceiveRemoveTrackable(FBackChannelOSCMessage& Message, FBackChannelOSCDispatch& Dispatch);
+
+	// Game thead message handlers
+	void ReceiveARInit_GameThread(FString ConfigObjectPathName, TArray<FARVideoFormat> Formats);
+	void ReceiveAddTrackable_GameThread(FString ClassPathName, TSharedPtr<TArray<uint8>, ESPMode::ThreadSafe> DataCopy);
+	void ReceiveUpdateTrackable_GameThread(FGuid UniqueId, TSharedPtr<TArray<uint8>, ESPMode::ThreadSafe> DataCopy);
+	void ReceiveRemoveTrackable_GameThread(FGuid UniqueId);
+
+	// AR notification handlers that send the data to the listener
+	void SendAddedMessage(UARTrackedGeometry* Added);
+	void SendUpdatedMessage(UARTrackedGeometry* Updated);
+	void SendRemovedMessage(UARTrackedGeometry* Removed);
+
+	/** Sends the name of the config object we're using so it can be loaded on the remote. Also passes the video settings and supported session types */
+	void SendARInitMessage();
+
+private:
+	/** The buffer that objects are serialized into. Cached to reduce malloc/free pressure */
+	TArray<uint8> SerializeBuffer;
+
+	/** Channel message handles */
+	FDelegateHandle InitMessageCallbackHandle;
+	FDelegateHandle AddMessageCallbackHandle;
+	FDelegateHandle UpdateMessageCallbackHandle;
+	FDelegateHandle RemoveMessageCallbackHandle;
+
+	/** AR system notification handles */
+	FDelegateHandle OnTrackableAddedDelegateHandle;
+	FDelegateHandle OnTrackableUpdatedDelegateHandle;
+	FDelegateHandle OnTrackableRemovedDelegateHandle;
+};
diff --git a/Engine/Plugins/Experimental/RemoteSession/Source/RemoteSession/Channels/RemoteSessionXRTrackingChannel.h b/Engine/Plugins/Experimental/RemoteSession/Source/RemoteSession/Channels/RemoteSessionXRTrackingChannel.h
--- a/Engine/Plugins/Experimental/RemoteSession/Source/RemoteSession/Channels/RemoteSessionXRTrackingChannel.h
+++ b/Engine/Plugins/Experimental/RemoteSession/Source/RemoteSession/Channels/RemoteSessionXRTrackingChannel.h
@@ -4,10 +4,10 @@
 
 #include "RemoteSessionChannel.h"
 #include "XRTrackingSystemBase.h"
+#include "ARSystem.h"
 
 class FBackChannelOSCMessage;
 class FBackChannelOSCDispatch;
-class IARSystemSupport;
 
 class REMOTESESSION_API FXRTrackingProxy :
 	public FXRTrackingSystemBase
diff --git a/Engine/Plugins/Experimental/RemoteSession/Source/RemoteSession/Private/Channels/RemoteSessionARSystemChannel.cpp b/Engine/Plugins/Experimental/RemoteSession/Source/RemoteSession/Private/Channels/RemoteSessionARSystemChannel.cpp
new file mode 100644
--- /dev/null
+++ b/Engine/Plugins/Experimental/RemoteSession/Source/RemoteSession/Private/Channels/RemoteSessionARSystemChannel.cpp
@@ -0,0 +1,405 @@
+// Copyright 1998-2019 Epic Games, Inc. All Rights Reserved.
+
+#include "Channels/RemoteSessionARSystemChannel.h"
+#include "RemoteSession.h"
+#include "BackChannel/Protocol/OSC/BackChannelOSCConnection.h"
+#include "BackChannel/Protocol/OSC/BackChannelOSCMessage.h"
+#include "MessageHandler/Messages.h"
+#include "ARSystem.h"
+#include "ARBlueprintLibrary.h"
+
+#include "Serialization/MemoryWriter.h"
+#include "Serialization/ObjectAndNameAsStringProxyArchive.h"
+
+DECLARE_CYCLE_STAT(TEXT("ARSystemChannel_Receive"), STAT_ARSystemChannel_Receive, STATGROUP_Game);
+
+#define INIT_MESSAGE_ADDRESS TEXT("/ARInit")
+#define ADD_TRACKABLE_MESSAGE_ADDRESS TEXT("/AddTrackable")
+#define UPDATE_TRACKABLE_MESSAGE_ADDRESS TEXT("/UpdateTrackable")
+#define REMOVE_TRACKABLE_MESSAGE_ADDRESS TEXT("/RemoveTrackable")
+
+TSharedPtr<FARSystemProxy, ESPMode::ThreadSafe> FARSystemProxy::FactoryInstance;
+
+FARSystemProxy::FARSystemProxy()
+	: SessionConfig(nullptr)
+{
+}
+
+FARSystemProxy::~FARSystemProxy()
+{
+}
+
+TSharedPtr<FARSystemProxy, ESPMode::ThreadSafe> FARSystemProxy::Get()
+{
+	if (!FactoryInstance.IsValid())
+	{
+		FactoryInstance = MakeShared<FARSystemProxy, ESPMode::ThreadSafe>();
+	}
+	return FactoryInstance;
+}
+
+IARSystemSupport* FARSystemProxy::GetARSystemPtr()
+{
+	return Get().Get();
+}
+
+void FARSystemProxy::Destroy()
+{
+	if (FactoryInstance.IsValid())
+	{
+		FactoryInstance.Reset();
+	}
+}
+
+void FARSystemProxy::AddReferencedObjects(FReferenceCollector& Collector)
+{
+	Collector.AddReferencedObject(SessionConfig);
+	Collector.AddReferencedObjects(TrackedGeometries);
+}
+
+EARTrackingQuality FARSystemProxy::OnGetTrackingQuality() const
+{
+	// @todo JoeG - Added a delegate for tracking quality so we can send that
+	return EARTrackingQuality::NotTracking;
+}
+
+FARSessionStatus FARSystemProxy::OnGetARSessionStatus() const
+{
+	// @todo JoeG - Added a delegate for this so we can send that
+	FARSessionStatus Status;
+	Status.Status = EARSessionStatus::Other;
+	return Status;
+}
+
+TArray<UARTrackedGeometry*> FARSystemProxy::OnGetAllTrackedGeometries() const
+{
+	check(IsInGameThread());
+
+	TArray<UARTrackedGeometry*> Geometries;
+	TrackedGeometries.GenerateValueArray(Geometries);
+	return Geometries;
+}
+
+bool FARSystemProxy::OnIsTrackingTypeSupported(EARSessionType SessionType) const
+{
+	return SessionConfig != nullptr && SessionConfig->GetSessionType() == SessionType;
+}
+
+EARWorldMappingState FARSystemProxy::OnGetWorldMappingStatus() const
+{
+	// @todo JoeG - Added a delegate for this so we can send that
+	return EARWorldMappingState::NotAvailable;
+}
+
+TArray<FARVideoFormat> FARSystemProxy::OnGetSupportedVideoFormats(EARSessionType SessionType) const
+{
+	check(IsInGameThread());
+
+	if (SessionConfig != nullptr && SessionConfig->GetSessionType() == SessionType)
+	{
+		return SupportedFormats;
+	}
+	return TArray<FARVideoFormat>();
+}
+
+void FARSystemProxy::SetSupportedVideoFormats(const TArray<FARVideoFormat>& InFormats)
+{
+	check(IsInGameThread());
+
+	SupportedFormats = InFormats;
+}
+
+void FARSystemProxy::SetSessionConfig(UARSessionConfig* InConfig)
+{
+	check(IsInGameThread());
+
+	SessionConfig = InConfig;
+}
+
+void FARSystemProxy::AddTrackable(UARTrackedGeometry* Added)
+{
+	check(IsInGameThread());
+
+	TrackedGeometries.Add(Added->UniqueId, Added);
+	TriggerOnTrackableAddedDelegates(Added);
+}
+
+UARTrackedGeometry* FARSystemProxy::GetTrackable(FGuid UniqueId)
+{
+	check(IsInGameThread());
+
+	UARTrackedGeometry** GeometrySearchResult = TrackedGeometries.Find(UniqueId);
+	return *GeometrySearchResult;
+}
+
+void FARSystemProxy::NotifyUpdated(UARTrackedGeometry* Updated)
+{
+	check(IsInGameThread());
+
+	TriggerOnTrackableUpdatedDelegates(Updated);
+}
+
+void FARSystemProxy::RemoveTrackable(FGuid UniqueId)
+{
+	check(IsInGameThread());
+
+	UARTrackedGeometry* TrackedGeometryBeingRemoved = TrackedGeometries.FindChecked(UniqueId);
+	TrackedGeometryBeingRemoved->UpdateTrackingState(EARTrackingState::StoppedTracking);
+
+	TriggerOnTrackableRemovedDelegates(TrackedGeometryBeingRemoved);
+
+	TrackedGeometries.Remove(UniqueId);
+}
+
+
+FRemoteSessionARSystemChannel::FRemoteSessionARSystemChannel(ERemoteSessionChannelMode InRole, TSharedPtr<FBackChannelOSCConnection, ESPMode::ThreadSafe> InConnection)
+	: FRemoteSessionXRTrackingChannel(InRole, InConnection, InRole == ERemoteSessionChannelMode::Read ? FARSystemProxy::GetARSystemPtr() : nullptr)
+{
+	// Are we receiving updates from the AR system? or sending them
+	if (Role == ERemoteSessionChannelMode::Read)
+	{
+		InitMessageCallbackHandle = Connection->AddMessageHandler(INIT_MESSAGE_ADDRESS, FBackChannelDispatchDelegate::FDelegate::CreateRaw(this, &FRemoteSessionARSystemChannel::ReceiveARInit));
+		Connection->SetMessageOptions(INIT_MESSAGE_ADDRESS, 1);
+
+		AddMessageCallbackHandle = Connection->AddMessageHandler(ADD_TRACKABLE_MESSAGE_ADDRESS, FBackChannelDispatchDelegate::FDelegate::CreateRaw(this, &FRemoteSessionARSystemChannel::ReceiveAddTrackable));
+		Connection->SetMessageOptions(ADD_TRACKABLE_MESSAGE_ADDRESS, 1000);
+
+		UpdateMessageCallbackHandle = Connection->AddMessageHandler(UPDATE_TRACKABLE_MESSAGE_ADDRESS, FBackChannelDispatchDelegate::FDelegate::CreateRaw(this, &FRemoteSessionARSystemChannel::ReceiveUpdateTrackable));
+		Connection->SetMessageOptions(UPDATE_TRACKABLE_MESSAGE_ADDRESS, 1000);
+
+		RemoveMessageCallbackHandle = Connection->AddMessageHandler(REMOVE_TRACKABLE_MESSAGE_ADDRESS, FBackChannelDispatchDelegate::FDelegate::CreateRaw(this, &FRemoteSessionARSystemChannel::ReceiveRemoveTrackable));
+		Connection->SetMessageOptions(REMOVE_TRACKABLE_MESSAGE_ADDRESS, 1000);
+	}
+	else
+	{
+		// Add the 3 AR trackable handlers
+		OnTrackableAddedDelegateHandle = UARBlueprintLibrary::AddOnTrackableAddedDelegate_Handle(FOnTrackableAddedDelegate::CreateRaw(this, &FRemoteSessionARSystemChannel::SendAddedMessage));
+		OnTrackableUpdatedDelegateHandle = UARBlueprintLibrary::AddOnTrackableUpdatedDelegate_Handle(FOnTrackableAddedDelegate::CreateRaw(this, &FRemoteSessionARSystemChannel::SendUpdatedMessage));
+		OnTrackableRemovedDelegateHandle = UARBlueprintLibrary::AddOnTrackableRemovedDelegate_Handle(FOnTrackableAddedDelegate::CreateRaw(this, &FRemoteSessionARSystemChannel::SendRemovedMessage));
+	}
+}
+
+FRemoteSessionARSystemChannel::~FRemoteSessionARSystemChannel()
+{
+	if (Role == ERemoteSessionChannelMode::Read)
+	{
+		// Clean up all of the message handlers
+		Connection->RemoveMessageHandler(INIT_MESSAGE_ADDRESS, InitMessageCallbackHandle);
+		Connection->RemoveMessageHandler(ADD_TRACKABLE_MESSAGE_ADDRESS, AddMessageCallbackHandle);
+		Connection->RemoveMessageHandler(UPDATE_TRACKABLE_MESSAGE_ADDRESS, UpdateMessageCallbackHandle);
+		Connection->RemoveMessageHandler(REMOVE_TRACKABLE_MESSAGE_ADDRESS, RemoveMessageCallbackHandle);
+	}
+	else
+	{
+		// Remove all of the AR notifications
+		UARBlueprintLibrary::ClearOnTrackableAddedDelegate_Handle(OnTrackableAddedDelegateHandle);
+		UARBlueprintLibrary::ClearOnTrackableUpdatedDelegate_Handle(OnTrackableUpdatedDelegateHandle);
+		UARBlueprintLibrary::ClearOnTrackableRemovedDelegate_Handle(OnTrackableRemovedDelegateHandle);
+	}
+
+	FARSystemProxy::Destroy();
+}
+
+void FRemoteSessionARSystemChannel::ReceiveARInit(FBackChannelOSCMessage& Message, FBackChannelOSCDispatch& Dispatch)
+{
+	check(ARSystemSupport != nullptr);
+
+	TArray<uint8> MsgData;
+	Message << MsgData;
+
+	FMemoryReader Ar(MsgData);
+	TwoParamMsg<FString, TArray<FARVideoFormat>> MsgParam(Ar);
+
+	UE_LOG(LogRemoteSession, Log, TEXT("Received AR session config (%s)"), *MsgParam.Param1);
+
+	// Since we are dealing with creating new UObjects, this needs to happen on the game thread
+	auto ARConfigTask = FSimpleDelegateGraphTask::FDelegate::CreateRaw(this, &FRemoteSessionARSystemChannel::ReceiveARInit_GameThread, MsgParam.Param1, MsgParam.Param2);
+	FSimpleDelegateGraphTask::CreateAndDispatchWhenReady(ARConfigTask, GET_STATID(STAT_ARSystemChannel_Receive), nullptr, ENamedThreads::GameThread);
+}
+
+void FRemoteSessionARSystemChannel::ReceiveARInit_GameThread(FString ConfigObjectPathName, TArray<FARVideoFormat> Formats)
+{
+	check(IsInGameThread());
+
+	FARSystemProxy::Get()->SetSupportedVideoFormats(Formats);
+
+	// Load and set the config object that was passed in
+	UARSessionConfig* SessionConfig = FindObject<UARSessionConfig>(ANY_PACKAGE, *ConfigObjectPathName);
+	if (SessionConfig == nullptr)
+	{
+		SessionConfig = LoadObject<UARSessionConfig>(nullptr, *ConfigObjectPathName);
+	}
+	// If the object could not be loaded (transient one) then create a default one and set that
+	if (SessionConfig == nullptr)
+	{
+		SessionConfig = NewObject<UARSessionConfig>();
+	}
+	FARSystemProxy::Get()->SetSessionConfig(SessionConfig);
+}
+
+void FRemoteSessionARSystemChannel::ReceiveAddTrackable(FBackChannelOSCMessage& Message, FBackChannelOSCDispatch& Dispatch)
+{
+	FString ClassPathName;
+	Message.Read(ClassPathName);
+	TSharedPtr<TArray<uint8>, ESPMode::ThreadSafe> DataCopy = MakeShareable(new TArray<uint8>());
+	Message.Read(*DataCopy);
+
+	// Since we are dealing with creating new UObjects, this needs to happen on the game thread
+	auto AddTrackableTask = FSimpleDelegateGraphTask::FDelegate::CreateRaw(this, &FRemoteSessionARSystemChannel::ReceiveAddTrackable_GameThread, ClassPathName, DataCopy);
+	FSimpleDelegateGraphTask::CreateAndDispatchWhenReady(AddTrackableTask, GET_STATID(STAT_ARSystemChannel_Receive), nullptr, ENamedThreads::GameThread);
+}
+
+void FRemoteSessionARSystemChannel::ReceiveAddTrackable_GameThread(FString ClassPathName, TSharedPtr<TArray<uint8>, ESPMode::ThreadSafe> DataCopy)
+{
+	UClass* TrackableClass = FindObject<UClass>(ANY_PACKAGE, *ClassPathName);
+	// We shouldn't have to load this since these are all native, but in case some AR platform has non-native classes...
+	if (TrackableClass == nullptr)
+	{
+		TrackableClass = LoadObject<UClass>(nullptr, *ClassPathName);
+	}
+
+	if (TrackableClass != nullptr)
+	{
+		UARTrackedGeometry* TrackedGeometry = NewObject<UARTrackedGeometry>(nullptr, TrackableClass);
+
+		FMemoryReader MemoryReader(*DataCopy, true);
+		FObjectAndNameAsStringProxyArchive Ar(MemoryReader, true);
+		TrackedGeometry->Serialize(Ar);
+
+		TrackedGeometry->SetLastUpdateTimestamp(FPlatformTime::Seconds());
+
+		FARSystemProxy::Get()->AddTrackable(TrackedGeometry);
+
+		UE_LOG(LogRemoteSession, Log, TEXT("Added new trackable (%s) with class (%s)"), *TrackedGeometry->GetName(), *ClassPathName);
+	}
+	else
+	{
+		UE_LOG(LogRemoteSession, Warning, TEXT("Couldn't find class (%s) for added trackable"), *ClassPathName);
+	}
+}
+
+void FRemoteSessionARSystemChannel::ReceiveUpdateTrackable(FBackChannelOSCMessage& Message, FBackChannelOSCDispatch& Dispatch)
+{
+	FGuid UniqueId;
+	FString StringGuid;
+	Message.Read(StringGuid);
+	FGuid::Parse(StringGuid, UniqueId);
+	TSharedPtr<TArray<uint8>, ESPMode::ThreadSafe> DataCopy = MakeShareable(new TArray<uint8>());
+	Message.Read(*DataCopy);
+
+	// Since we are dealing with updating UObjects, this needs to happen on the game thread
+	auto UpdateTrackableTask = FSimpleDelegateGraphTask::FDelegate::CreateRaw(this, &FRemoteSessionARSystemChannel::ReceiveUpdateTrackable_GameThread, UniqueId, DataCopy);
+	FSimpleDelegateGraphTask::CreateAndDispatchWhenReady(UpdateTrackableTask, GET_STATID(STAT_ARSystemChannel_Receive), nullptr, ENamedThreads::GameThread);
+}
+
+void FRemoteSessionARSystemChannel::ReceiveUpdateTrackable_GameThread(FGuid UniqueId, TSharedPtr<TArray<uint8>, ESPMode::ThreadSafe> DataCopy)
+{
+	check(IsInGameThread());
+
+	UARTrackedGeometry* Updated = FARSystemProxy::Get()->GetTrackable(UniqueId);
+	if (Updated != nullptr)
+	{
+		FMemoryReader MemoryReader(*DataCopy, true);
+		FObjectAndNameAsStringProxyArchive Ar(MemoryReader, true);
+		Updated->Serialize(Ar);
+
+		Updated->SetLastUpdateTimestamp(FPlatformTime::Seconds());
+
+		FARSystemProxy::Get()->NotifyUpdated(Updated);
+
+		UE_LOG(LogRemoteSession, Log, TEXT("Updated trackable (%s) with UniqueId (%s)"), *Updated->GetName(), *UniqueId.ToString());
+	}
+	else
+	{
+		UE_LOG(LogRemoteSession, Warning, TEXT("Couldn't find trackable with UniqueId (%s)"), *UniqueId.ToString());
+	}
+}
+
+void FRemoteSessionARSystemChannel::ReceiveRemoveTrackable(FBackChannelOSCMessage& Message, FBackChannelOSCDispatch& Dispatch)
+{
+	FGuid UniqueId;
+	FString StringGuid;
+	Message.Read(StringGuid);
+	FGuid::Parse(StringGuid, UniqueId);
+
+	// Since we are dealing with updating UObjects, this needs to happen on the game thread
+	auto RemoveTrackableTask = FSimpleDelegateGraphTask::FDelegate::CreateRaw(this, &FRemoteSessionARSystemChannel::ReceiveRemoveTrackable_GameThread, UniqueId);
+	FSimpleDelegateGraphTask::CreateAndDispatchWhenReady(RemoveTrackableTask, GET_STATID(STAT_ARSystemChannel_Receive), nullptr, ENamedThreads::GameThread);
+}
+
+void FRemoteSessionARSystemChannel::ReceiveRemoveTrackable_GameThread(FGuid UniqueId)
+{
+	check(IsInGameThread());
+
+	FARSystemProxy::Get()->RemoveTrackable(UniqueId);
+
+	UE_LOG(LogRemoteSession, Log, TEXT("Removed trackable with UniqueId (%s)"), *UniqueId.ToString());
+}
+
+void FRemoteSessionARSystemChannel::SendARInitMessage()
+{
+	if (Connection.IsValid())
+	{
+		UARSessionConfig* Config = UARBlueprintLibrary::GetSessionConfig();
+		if (Config != nullptr)
+		{
+			FString PathName = Config->GetPathName();
+			TArray<FARVideoFormat> SupportedFormats = UARBlueprintLibrary::GetSupportedVideoFormats(Config->GetSessionType());
+
+			TwoParamMsg<FString, TArray<FARVideoFormat>> MsgParam(PathName, SupportedFormats);
+			FBackChannelOSCMessage Msg(INIT_MESSAGE_ADDRESS);
+			Msg.Write(MsgParam.AsData());
+
+			Connection->SendPacket(Msg);
+
+			UE_LOG(LogRemoteSession, Log, TEXT("Sent AR init message with session config (%s)"), *PathName);
+		}
+	}
+}
+
+void FRemoteSessionARSystemChannel::SendAddedMessage(UARTrackedGeometry* Added)
+{
+	// Serialize the object into a buffer for sending
+	SerializeBuffer.Reset();
+	FMemoryWriter MemoryWriter(SerializeBuffer, true);
+	FObjectAndNameAsStringProxyArchive Ar(MemoryWriter, false);
+	Added->Serialize(Ar);
+
+	FString ClassPathName = Added->GetClass()->GetPathName();
+
+	FBackChannelOSCMessage Msg(ADD_TRACKABLE_MESSAGE_ADDRESS);
+	Msg.Write(ClassPathName);
+	Msg.Write(SerializeBuffer);
+	Connection->SendPacket(Msg);
+
+	UE_LOG(LogRemoteSession, Log, TEXT("Sent trackable added (%s)"), *Added->GetName());
+}
+
+void FRemoteSessionARSystemChannel::SendUpdatedMessage(UARTrackedGeometry* Updated)
+{
+	// Serialize the object into a buffer for sending
+	SerializeBuffer.Reset();
+	FMemoryWriter MemoryWriter(SerializeBuffer, true);
+	FObjectAndNameAsStringProxyArchive Ar(MemoryWriter, false);
+	Updated->Serialize(Ar);
+
+	FString TrackableGuid = Updated->UniqueId.ToString();
+
+	FBackChannelOSCMessage Msg(UPDATE_TRACKABLE_MESSAGE_ADDRESS);
+	Msg.Write(TrackableGuid);
+	Msg.Write(SerializeBuffer);
+	Connection->SendPacket(Msg);
+
+	UE_LOG(LogRemoteSession, Log, TEXT("Sent trackable updated (%s)"), *Updated->GetName());
+}
+
+void FRemoteSessionARSystemChannel::SendRemovedMessage(UARTrackedGeometry* Removed)
+{
+	FString TrackableGuid = Removed->UniqueId.ToString();
+
+	FBackChannelOSCMessage Msg(UPDATE_TRACKABLE_MESSAGE_ADDRESS);
+	Msg.Write(TrackableGuid);
+	Connection->SendPacket(Msg);
+
+	UE_LOG(LogRemoteSession, Log, TEXT("Sent trackable removed (%s)"), *Removed->GetName());
+}
\ No newline at end of file
diff --git a/Engine/Plugins/Experimental/RemoteSession/Source/RemoteSession/Private/RemoteSessionRole.cpp b/Engine/Plugins/Experimental/RemoteSession/Source/RemoteSession/Private/RemoteSessionRole.cpp
--- a/Engine/Plugins/Experimental/RemoteSession/Source/RemoteSession/Private/RemoteSessionRole.cpp
+++ b/Engine/Plugins/Experimental/RemoteSession/Source/RemoteSession/Private/RemoteSessionRole.cpp
@@ -7,6 +7,7 @@
 #include "Channels/RemoteSessionInputChannel.h"
 #include "Channels/RemoteSessionXRTrackingChannel.h"
 #include "Channels/RemoteSessionARCameraChannel.h"
+#include "Channels/RemoteSessionARSystemChannel.h"
 #include "Channels/RemoteSessionFrameBufferChannel.h"
 #include "Async/Async.h"
 #include "GeneralProjectSettings.h"
@@ -223,7 +224,7 @@ void FRemoteSessionRole::CreateChannel(const FString& InChannelName, ERemoteSess
 	}
 	else if (InChannelName == FRemoteSessionXRTrackingChannel::StaticType())
 	{
-		bool IsSupported = (InMode == ERemoteSessionChannelMode::Read) || UARBlueprintLibrary::IsSessionTypeSupported(EARSessionType::World);
+		bool IsSupported = (InMode == ERemoteSessionChannelMode::Read) || UARBlueprintLibrary::IsSessionTypeSupported(EARSessionType::Orientation);
 		if (IsSupported)
 		{
 			NewChannel = MakeShareable(new FRemoteSessionXRTrackingChannel(InMode, OSCConnection));
@@ -247,7 +248,19 @@ void FRemoteSessionRole::CreateChannel(const FString& InChannelName, ERemoteSess
 			UE_LOG(LogRemoteSession, Warning, TEXT("FRemoteSessionARCameraChannel does not support sending on this platform/device"));
 		}
 	}
-	
+	else if (InChannelName == FRemoteSessionARSystemChannel::StaticType())
+	{
+		bool IsSupported = (InMode == ERemoteSessionChannelMode::Read) || UARBlueprintLibrary::IsSessionTypeSupported(EARSessionType::World);
+		if (IsSupported)
+		{
+			NewChannel = MakeShareable(new FRemoteSessionARSystemChannel(InMode, OSCConnection));
+		}
+		else
+		{
+			UE_LOG(LogRemoteSession, Warning, TEXT("FRemoteSessionARSystemChannel does not support sending on this platform/device"));
+		}
+	}
+
 	if (NewChannel.IsValid())
 	{
 		UE_LOG(LogRemoteSession, Log, TEXT("Created Channel %s with mode %d"), *InChannelName, (int)InMode);
diff --git a/Engine/Source/Runtime/AugmentedReality/Public/ARTrackable.h b/Engine/Source/Runtime/AugmentedReality/Public/ARTrackable.h
--- a/Engine/Source/Runtime/AugmentedReality/Public/ARTrackable.h
+++ b/Engine/Source/Runtime/AugmentedReality/Public/ARTrackable.h
@@ -53,6 +53,7 @@ public:
 	
 	UFUNCTION(BlueprintPure, Category="AR AugmentedReality|Tracked Geometry")
 	float GetLastUpdateTimestamp() const;
+	inline void SetLastUpdateTimestamp(double InTimestamp) { LastUpdateTimestamp = InTimestamp; }
 
 	UPROPERTY(BlueprintReadOnly, Category="AR AugmentedReality|Tracked Geometry")
 	FGuid UniqueId;
diff --git a/Engine/Source/Runtime/AugmentedReality/Public/ARTypes.h b/Engine/Source/Runtime/AugmentedReality/Public/ARTypes.h
--- a/Engine/Source/Runtime/AugmentedReality/Public/ARTypes.h
+++ b/Engine/Source/Runtime/AugmentedReality/Public/ARTypes.h
@@ -401,5 +401,10 @@ public:
 	int32 Height;
 
 	bool IsValidFormat() { return FPS > 0 && Width > 0 && Height > 0; }
+
+	friend FArchive& operator<<(FArchive& Ar, FARVideoFormat& Format)
+	{
+		return Ar << Format.FPS << Format.Width << Format.Height;
+	}
 };
 
