Fixed the confusion between bits per pixel when the API wants bytes per pixel

diff --git a/Engine/Plugins/Experimental/RemoteSession/Source/RemoteSession/Private/Channels/RemoteSessionFrameBufferChannel.cpp b/Engine/Plugins/Experimental/RemoteSession/Source/RemoteSession/Private/Channels/RemoteSessionFrameBufferChannel.cpp
--- a/Engine/Plugins/Experimental/RemoteSession/Source/RemoteSession/Private/Channels/RemoteSessionFrameBufferChannel.cpp
+++ b/Engine/Plugins/Experimental/RemoteSession/Source/RemoteSession/Private/Channels/RemoteSessionFrameBufferChannel.cpp
@@ -274,7 +274,7 @@ void FRemoteSessionFrameBufferChannel::Tick(const float InDeltaTime)
 				delete InRegions; 
 			};
 
-			DecodedTextures[NextImage]->UpdateTextureRegions(0, 1, Region, 4 * QueuedImage->Width, 8, TextureData->GetData(), DataCleanupFunc);
+			DecodedTextures[NextImage]->UpdateTextureRegions(0, 1, Region, 4 * QueuedImage->Width, sizeof(FColor), TextureData->GetData(), DataCleanupFunc);
 
 			UE_LOG(LogRemoteSession, Verbose, TEXT("GT: Uploaded image %d"),
 				QueuedImage->ImageIndex);
