Made the same bits vs bytes per pixel fix in the AR camera streaming

diff --git a/Engine/Plugins/Experimental/RemoteSession/Source/RemoteSession/Private/Channels/RemoteSessionARCameraChannel.cpp b/Engine/Plugins/Experimental/RemoteSession/Source/RemoteSession/Private/Channels/RemoteSessionARCameraChannel.cpp
--- a/Engine/Plugins/Experimental/RemoteSession/Source/RemoteSession/Private/Channels/RemoteSessionARCameraChannel.cpp
+++ b/Engine/Plugins/Experimental/RemoteSession/Source/RemoteSession/Private/Channels/RemoteSessionARCameraChannel.cpp
@@ -570,7 +570,7 @@ void FRemoteSessionARCameraChannel::UpdateRenderingTexture()
 		FUpdateTextureRegion2D* Region = new FUpdateTextureRegion2D(0, 0, 0, 0, DecompressedImage->Width, DecompressedImage->Height);
 		TArray<uint8>* TextureData = new TArray<uint8>(MoveTemp(DecompressedImage->ImageData));
 
-		RenderingTextures[NextImage]->UpdateTextureRegions(0, 1, Region, 4 * DecompressedImage->Width, 8, TextureData->GetData(), [this, NextImage](auto InTextureData, auto InRegions)
+		RenderingTextures[NextImage]->UpdateTextureRegions(0, 1, Region, 4 * DecompressedImage->Width, sizeof(FColor), TextureData->GetData(), [this, NextImage](auto InTextureData, auto InRegions)
 		{
 			RenderingTextureIndex.Set(NextImage);
 			RenderingTexturesUpdateCount[NextImage].Decrement();
