Fix threading issue with RemoveSession. Add an ImageSender to the ImageChannel to prevent multithread access.

diff --git a/Engine/Plugins/Experimental/RemoteSession/Source/RemoteSession/Channels/RemoteSessionFrameBufferChannel.h b/Engine/Plugins/Experimental/RemoteSession/Source/RemoteSession/Channels/RemoteSessionFrameBufferChannel.h
--- a/Engine/Plugins/Experimental/RemoteSession/Source/RemoteSession/Channels/RemoteSessionFrameBufferChannel.h
+++ b/Engine/Plugins/Experimental/RemoteSession/Source/RemoteSession/Channels/RemoteSessionFrameBufferChannel.h
@@ -5,7 +5,7 @@
 #include "RemoteSessionImageChannel.h"
 
 // FRemoteSessionFrameBufferChannel is deprecated. Please use FRemoteSessionImageChannel.
-// FRemoteSessionFrameBufferChannelFactoryWorker was created for backward comptability with older app
+// FRemoteSessionFrameBufferChannelFactoryWorker was created for backward compatibility with older app
 
 class REMOTESESSION_API FRemoteSessionFrameBufferChannelFactoryWorker : public IRemoteSessionChannelFactoryWorker
 {
diff --git a/Engine/Plugins/Experimental/RemoteSession/Source/RemoteSession/Channels/RemoteSessionImageChannel.h b/Engine/Plugins/Experimental/RemoteSession/Source/RemoteSession/Channels/RemoteSessionImageChannel.h
--- a/Engine/Plugins/Experimental/RemoteSession/Source/RemoteSession/Channels/RemoteSessionImageChannel.h
+++ b/Engine/Plugins/Experimental/RemoteSession/Source/RemoteSession/Channels/RemoteSessionImageChannel.h
@@ -23,11 +23,42 @@ public:
 };
 
 /**
- *	A channel that take an image, encodes it as a jpg in an async task, then sends it to the client.
+ *	A channel that takes an image (created by an IRemoteSessionImageProvider), then sends it to the client (with a FRemoteSessionImageSender).
  *	On the client images are decoded into a double-buffered texture that can be accessed via GetHostScreen.
  */
 class REMOTESESSION_API FRemoteSessionImageChannel : public IRemoteSessionChannel, FRunnable
 {
+public:
+
+	/**
+	*	A helper object responsible to take the raw data, encode it to jpg and send it to the client for the RemoteSessionImageChannel
+	*/
+	class FImageSender
+	{
+	public:
+  
+		FImageSender(TSharedPtr<FBackChannelOSCConnection, ESPMode::ThreadSafe> InConnection);
+  
+		/** Set the jpg compression quality */
+		void SetCompressQuality(int32 InQuality);
+  
+		/** Send an image to the connected clients */
+		void SendRawImageToClients(int32 Width, int32 Height, const TArray<FColor>& ImageData);
+  
+		/** Send an image to the connected clients */
+		void SendRawImageToClients(int32 Width, int32 Height, const void* ImageData, int32 AllocatedImageDataSize);
+  
+	private:
+  
+		/** Underlying connection */
+		TWeakPtr<FBackChannelOSCConnection, ESPMode::ThreadSafe> Connection;
+  
+		/** Compression quality of the raw image we wish to send to client */
+		TAtomic<int32> CompressQuality;
+  
+		int32 NumSentImages;
+	};
+
 public:
 
 	FRemoteSessionImageChannel(ERemoteSessionChannelMode InRole, TSharedPtr<FBackChannelOSCConnection, ESPMode::ThreadSafe> InConnection);
@@ -46,11 +77,8 @@ public:
 	/** Set the jpg compression quality */
 	void SetCompressQuality(int32 InQuality);
 
-	/** Send an image to the connected clients */
-	void SendRawImageToClients(int32 Width, int32 Height, const TArray<FColor>& ImageData);
-
-	/** Send an image to the connected clients */
-	void SendRawImageToClients(int32 Width, int32 Height, const void* ImageData, int32 AllocatedImageDataSize);
+	/** Return the image sender connected to the clients */
+	TSharedPtr<FRemoteSessionImageChannel::FImageSender, ESPMode::ThreadSafe> GetImageSender() { return ImageSender; }
 
 	/* Begin IRemoteSessionChannel implementation */
 	static const TCHAR* StaticType() { return TEXT("FRemoteSessionImageChannel"); }
@@ -89,18 +117,16 @@ protected:
 	};
 
 	FCriticalSection										IncomingImageMutex;
-	TArray<TSharedPtr<FImageData, ESPMode::ThreadSafe>>		IncomingEncodedImages;
+	TArray<TUniquePtr<FImageData>>							IncomingEncodedImages;
 
 	FCriticalSection										DecodedImageMutex;
-	TArray<TSharedPtr<FImageData, ESPMode::ThreadSafe>>		IncomingDecodedImages;
+	TArray<TUniquePtr<FImageData>>							IncomingDecodedImages;
 
 	UTexture2D*												DecodedTextures[2];
-	int32													DecodedTextureIndex;
-
-	int32													NumSentImages;
+	TSharedPtr<TAtomic<int32>, ESPMode::ThreadSafe>			DecodedTextureIndex;
 
-	/** Compression quality of the raw image we wish to send to client */
-	int32 CompressQuality;
+	/** Image sender used by the channel */
+	TSharedPtr<FImageSender, ESPMode::ThreadSafe> ImageSender;
 
 	/** So we can manage callback lifetimes properly */
 	FDelegateHandle MessageCallbackHandle;
diff --git a/Engine/Plugins/Experimental/RemoteSession/Source/RemoteSession/ImageProviders/RemoteSessionFrameBufferImageProvider.h b/Engine/Plugins/Experimental/RemoteSession/Source/RemoteSession/ImageProviders/RemoteSessionFrameBufferImageProvider.h
--- a/Engine/Plugins/Experimental/RemoteSession/Source/RemoteSession/ImageProviders/RemoteSessionFrameBufferImageProvider.h
+++ b/Engine/Plugins/Experimental/RemoteSession/Source/RemoteSession/ImageProviders/RemoteSessionFrameBufferImageProvider.h
@@ -15,7 +15,7 @@ class REMOTESESSION_API FRemoteSessionFrameBufferImageProvider : public IRemoteS
 {
 public:
 
-	FRemoteSessionFrameBufferImageProvider(TWeakPtr<FRemoteSessionImageChannel> Owner);
+	FRemoteSessionFrameBufferImageProvider(TSharedPtr<FRemoteSessionImageChannel::FImageSender, ESPMode::ThreadSafe> ImageSender);
 	~FRemoteSessionFrameBufferImageProvider();
 
 	/** Specifies which viewport to capture */
@@ -38,11 +38,11 @@ protected:
 	/** Release the FrameGrabber*/
 	void ReleaseFrameGrabber();
 
-	TWeakPtr<FRemoteSessionImageChannel> ImageChannel;
+	TWeakPtr<FRemoteSessionImageChannel::FImageSender, ESPMode::ThreadSafe> ImageSender;
 
 	TSharedPtr<FFrameGrabber> FrameGrabber;
 
-	FThreadSafeCounter NumDecodingTasks;
+	TSharedPtr<FThreadSafeCounter, ESPMode::ThreadSafe> NumDecodingTasks;
 
 	/** Time we last sent an image */
 	double LastSentImageTime;
diff --git a/Engine/Plugins/Experimental/RemoteSession/Source/RemoteSession/Private/Channels/RemoteSessionFrameBufferChannel.cpp b/Engine/Plugins/Experimental/RemoteSession/Source/RemoteSession/Private/Channels/RemoteSessionFrameBufferChannel.cpp
--- a/Engine/Plugins/Experimental/RemoteSession/Source/RemoteSession/Private/Channels/RemoteSessionFrameBufferChannel.cpp
+++ b/Engine/Plugins/Experimental/RemoteSession/Source/RemoteSession/Private/Channels/RemoteSessionFrameBufferChannel.cpp
@@ -12,7 +12,7 @@ TSharedPtr<IRemoteSessionChannel> FRemoteSessionFrameBufferChannelFactoryWorker:
 
 	if (InMode == ERemoteSessionChannelMode::Write)
 	{
-		TSharedPtr<FRemoteSessionFrameBufferImageProvider> ImageProvider = MakeShared<FRemoteSessionFrameBufferImageProvider>(Channel);
+		TSharedPtr<FRemoteSessionFrameBufferImageProvider> ImageProvider = MakeShared<FRemoteSessionFrameBufferImageProvider>(Channel->GetImageSender());
 
 		{
 			int32 Quality = 85;
diff --git a/Engine/Plugins/Experimental/RemoteSession/Source/RemoteSession/Private/Channels/RemoteSessionImageChannel.cpp b/Engine/Plugins/Experimental/RemoteSession/Source/RemoteSession/Private/Channels/RemoteSessionImageChannel.cpp
--- a/Engine/Plugins/Experimental/RemoteSession/Source/RemoteSession/Private/Channels/RemoteSessionImageChannel.cpp
+++ b/Engine/Plugins/Experimental/RemoteSession/Source/RemoteSession/Private/Channels/RemoteSessionImageChannel.cpp
@@ -24,15 +24,66 @@
 	TEXT("Sets quality (1-100)"),
 	ECVF_Default);
 
+
+FRemoteSessionImageChannel::FImageSender::FImageSender(TSharedPtr<FBackChannelOSCConnection, ESPMode::ThreadSafe> InConnection)
+{
+	Connection = InConnection;
+	CompressQuality = 0;
+	NumSentImages = 0;
+}
+
+void FRemoteSessionImageChannel::FImageSender::SetCompressQuality(int32 InQuality)
+{
+	CompressQuality = InQuality;
+}
+
+void FRemoteSessionImageChannel::FImageSender::SendRawImageToClients(int32 Width, int32 Height, const TArray<FColor>& ImageData)
+{
+	SendRawImageToClients(Width, Height, ImageData.GetData(), ImageData.GetAllocatedSize());
+}
+
+void FRemoteSessionImageChannel::FImageSender::SendRawImageToClients(int32 Width, int32 Height, const void* ImageData, int32 AllocatedImageDataSize)
+{
+	static bool SkipImages = FParse::Param(FCommandLine::Get(), TEXT("remote.noimage"));
+
+	// Can be released on the main thread at anytime so hold onto it
+	TSharedPtr<FBackChannelOSCConnection, ESPMode::ThreadSafe> LocalConnection = Connection.Pin();
+
+	if (LocalConnection.IsValid() && SkipImages == false)
+	{
+		const double TimeNow = FPlatformTime::Seconds();
+
+		// created on demand because there can be multiple SendImage requests in flight
+		IImageWrapperModule* ImageWrapperModule = FModuleManager::GetModulePtr<IImageWrapperModule>(FName("ImageWrapper"));
+		if (ImageWrapperModule != nullptr)
+		{
+			TSharedPtr<IImageWrapper> ImageWrapper = ImageWrapperModule->CreateImageWrapper(EImageFormat::JPEG);
+
+			ImageWrapper->SetRaw(ImageData, AllocatedImageDataSize, Width, Height, ERGBFormat::BGRA, 8);
+
+			const int32 CurrentQuality = QualityMasterSetting > 0 ? QualityMasterSetting : CompressQuality;
+			const TArray<uint8>& JPGData = ImageWrapper->GetCompressed(CompressQuality);
+
+			FBackChannelOSCMessage Msg(TEXT("/Screen"));
+			Msg.Write(Width);
+			Msg.Write(Height);
+			Msg.Write(JPGData);
+			Msg.Write(++NumSentImages);
+			LocalConnection->SendPacket(Msg);
+
+			UE_LOG(LogRemoteSession, Verbose, TEXT("Sent image %d in %.02f ms"),
+				NumSentImages, (FPlatformTime::Seconds() - TimeNow) * 1000.0);
+		}
+	}
+}
+
 FRemoteSessionImageChannel::FRemoteSessionImageChannel(ERemoteSessionChannelMode InRole, TSharedPtr<FBackChannelOSCConnection, ESPMode::ThreadSafe> InConnection)
 	: IRemoteSessionChannel(InRole, InConnection)
 {
 	Connection = InConnection;
 	DecodedTextures[0] = nullptr;
 	DecodedTextures[1] = nullptr;
-	DecodedTextureIndex = 0;
-	NumSentImages = 0;
-	CompressQuality = 0;
+	DecodedTextureIndex = MakeShared<TAtomic<int32>, ESPMode::ThreadSafe>(0);
 	Role = InRole;
 	LastDecodedImageIndex = 0;
 	LastIncomingImageIndex = 0;
@@ -50,6 +101,10 @@
 
 		StartBackgroundThread();
 	}
+	else
+	{
+		ImageSender = MakeShared<FRemoteSessionImageChannel::FImageSender, ESPMode::ThreadSafe>(InConnection);
+	}
 }
 
 FRemoteSessionImageChannel::~FRemoteSessionImageChannel()
@@ -93,7 +148,7 @@
 
 UTexture2D* FRemoteSessionImageChannel::GetHostScreen() const
 {
-	return DecodedTextures[DecodedTextureIndex];
+	return DecodedTextures[DecodedTextureIndex->Load()];
 }
 
 void FRemoteSessionImageChannel::Tick(const float InDeltaTime)
@@ -112,7 +167,7 @@
 	{
 		SCOPE_CYCLE_COUNTER(STAT_TextureUpdate);
 
-		TSharedPtr<FImageData, ESPMode::ThreadSafe> QueuedImage;
+		TUniquePtr<FImageData> QueuedImage;
 
 		{
 			// Check to see if there are any queued images. We just care about the last
@@ -120,20 +175,20 @@
 			if (IncomingDecodedImages.Num())
 			{
 				INC_DWORD_STAT(STAT_RSNumFrames);
-				QueuedImage = IncomingDecodedImages.Last();
+				QueuedImage = MoveTemp(IncomingDecodedImages.Last());
 				LastDecodedImageIndex = QueuedImage->ImageIndex;
 
 				UE_LOG(LogRemoteSession, Verbose, TEXT("GT: Image %d is ready, discarding %d earlier images"),
 					QueuedImage->ImageIndex, IncomingDecodedImages.Num()-1);
 
-				IncomingDecodedImages.Empty();
+				IncomingDecodedImages.Reset();
 			}
 		}
 
 		// If an image was waiting...
 		if (QueuedImage.IsValid())
 		{
-			int32 NextImage = DecodedTextureIndex == 0 ? 1 : 0;
+			int32 NextImage = DecodedTextureIndex->Load() == 0 ? 1 : 0;
 
 			// create a texture if we don't have a suitable one
 			if (DecodedTextures[NextImage] == nullptr || QueuedImage->Width != DecodedTextures[NextImage]->GetSizeX() || QueuedImage->Height != DecodedTextures[NextImage]->GetSizeY())
@@ -144,14 +199,18 @@
 			// Update it on the render thread. There shouldn't (...) be any harm in GT code using it from this point
 			FUpdateTextureRegion2D* Region = new FUpdateTextureRegion2D(0, 0, 0, 0, QueuedImage->Width, QueuedImage->Height);
 			TArray<uint8>* TextureData = new TArray<uint8>(MoveTemp(QueuedImage->ImageData));
+			TWeakPtr<TAtomic<int32>, ESPMode::ThreadSafe> DecodedTextureIndexWeaked = DecodedTextureIndex;
 
 			// cleanup functions, gets executed on the render thread after UpdateTextureRegions
-			TFunction<void(uint8*, const FUpdateTextureRegion2D*)> DataCleanupFunc = [this, NextImage, TextureData](uint8* InTextureData, const FUpdateTextureRegion2D* InRegions)
+			TFunction<void(uint8*, const FUpdateTextureRegion2D*)> DataCleanupFunc = [DecodedTextureIndexWeaked, NextImage, TextureData](uint8* InTextureData, const FUpdateTextureRegion2D* InRegions)
 			{
-				DecodedTextureIndex = NextImage;
+				if (TSharedPtr<TAtomic<int32>, ESPMode::ThreadSafe> DecodedTextureIndexPinned = DecodedTextureIndexWeaked.Pin())
+				{
+					DecodedTextureIndexPinned->Store(NextImage);
+				}
 
 				//this is executed later on the render thread, meanwhile TextureData might have changed
-				delete InTextureData; 
+				delete TextureData;
 				delete InRegions; 
 			};
 
@@ -175,68 +234,30 @@
 {
 	if (Role == ERemoteSessionChannelMode::Write)
 	{
-		CompressQuality = InQuality;
+		ImageSender->SetCompressQuality(InQuality);
 	}
 }
 
-void FRemoteSessionImageChannel::SendRawImageToClients(int32 Width, int32 Height, const TArray<FColor>& ImageData)
-{
-	SendRawImageToClients(Width, Height, ImageData.GetData(), ImageData.GetAllocatedSize());
-}
-
-void FRemoteSessionImageChannel::SendRawImageToClients(int32 Width, int32 Height, const void* ImageData, int32 AllocatedImageDataSize)
-{
-	static bool SkipImages = FParse::Param(FCommandLine::Get(), TEXT("remote.noimage"));
-
-	// Can be released on the main thread at anytime so hold onto it
-	TSharedPtr<FBackChannelOSCConnection, ESPMode::ThreadSafe> LocalConnection = Connection.Pin();
-
-	if (LocalConnection.IsValid() && SkipImages == false)
-	{
-		const double TimeNow = FPlatformTime::Seconds();
-
-		// created on demand because there can be multiple SendImage requests in flight
-		IImageWrapperModule* ImageWrapperModule = FModuleManager::GetModulePtr<IImageWrapperModule>(FName("ImageWrapper"));
-		if (ImageWrapperModule != nullptr)
-		{
-			TSharedPtr<IImageWrapper> ImageWrapper = ImageWrapperModule->CreateImageWrapper(EImageFormat::JPEG);
-
-			ImageWrapper->SetRaw(ImageData, AllocatedImageDataSize, Width, Height, ERGBFormat::BGRA, 8);
-
-			const int32 CurrentQuality = QualityMasterSetting > 0 ? QualityMasterSetting : CompressQuality;
-			const TArray<uint8>& JPGData = ImageWrapper->GetCompressed(CompressQuality);
-
-			FBackChannelOSCMessage Msg(TEXT("/Screen"));
-			Msg.Write(Width);
-			Msg.Write(Height);
-			Msg.Write(JPGData);
-			Msg.Write(++NumSentImages);
-			LocalConnection->SendPacket(Msg);
-
-			UE_LOG(LogRemoteSession, Verbose, TEXT("Sent image %d in %.02f ms"),
-				NumSentImages, (FPlatformTime::Seconds() - TimeNow) * 1000.0);
-		}
-	}
-}
-
 void FRemoteSessionImageChannel::ReceiveHostImage(FBackChannelOSCMessage& Message, FBackChannelOSCDispatch& Dispatch)
 {
-	int32 Width(0);
-	int32 Height(0);
+	int32 ImageIndex(0);
 
-	TSharedPtr<FImageData, ESPMode::ThreadSafe> ReceivedImage = MakeShareable(new FImageData);
+	TUniquePtr<FImageData> ReceivedImage = MakeUnique<FImageData>();
 
 	Message << ReceivedImage->Width;
 	Message << ReceivedImage->Height;
 	Message << ReceivedImage->ImageData;
-	Message << ReceivedImage->ImageIndex;
+	Message << ImageIndex;
 
+	ReceivedImage->ImageIndex = ImageIndex;
+
 	FScopeLock Lock(&IncomingImageMutex);
 	if (LastIncomingImageIndex > 0 && IncomingEncodedImages.Num() > 1 && IncomingEncodedImages.Last()->ImageIndex > LastIncomingImageIndex)
 	{
 		IncomingEncodedImages.RemoveSingle(IncomingEncodedImages.Last());
 	}
-	IncomingEncodedImages.Add(ReceivedImage);
 
+	IncomingEncodedImages.Add(MoveTemp(ReceivedImage));
+
 	if (ScreenshotEvent)
 	{
@@ -244,14 +265,13 @@
 		ScreenshotEvent->Trigger();
 	}
 
-	UE_LOG(LogRemoteSession, Verbose, TEXT("Received Image %d, %d pending"), 
-		ReceivedImage->ImageIndex, IncomingEncodedImages.Num());
-
+	UE_LOG(LogRemoteSession, Verbose, TEXT("Received Image %d, %d pending"),
+		ImageIndex, IncomingEncodedImages.Num());
 }
 
 void FRemoteSessionImageChannel::ProcessIncomingTextures()
 {
-	TSharedPtr<FImageData, ESPMode::ThreadSafe> Image;
+	TUniquePtr<FImageData> Image;
 	const double StartTime = FPlatformTime::Seconds();
 	{
 		// check if there's anything to do, if not pause the background thread
@@ -263,13 +283,13 @@
 		}
 
 		// take the last image we don't care about the rest
-		Image = IncomingEncodedImages.Last();
+		Image = MoveTemp(IncomingEncodedImages.Last());
 		LastIncomingImageIndex = Image->ImageIndex;
 
 		UE_LOG(LogRemoteSession, Verbose, TEXT("Processing Image %d, discarding %d other pending images"), 
 			Image->ImageIndex, IncomingEncodedImages.Num() - 1);
 
-		IncomingEncodedImages.Empty();
+		IncomingEncodedImages.Reset();
 	}
 
 	IImageWrapperModule* ImageWrapperModule = FModuleManager::GetModulePtr<IImageWrapperModule>(FName("ImageWrapper"));
@@ -284,7 +304,7 @@
 
 		if (ImageWrapper->GetRaw(ERGBFormat::BGRA, 8, RawData))
 		{
-			TSharedPtr<FImageData, ESPMode::ThreadSafe> QueuedImage = MakeShareable(new FImageData);
+			TUniquePtr<FImageData> QueuedImage = MakeUnique<FImageData>();
 			QueuedImage->Width = Image->Width;
 			QueuedImage->Height = Image->Height;
 			QueuedImage->ImageData = MoveTemp(*((TArray<uint8>*)RawData));
@@ -296,7 +316,7 @@
 				{
 					IncomingDecodedImages.RemoveSingle(IncomingDecodedImages.Last());
 				}
-				IncomingDecodedImages.Add(QueuedImage);
+				IncomingDecodedImages.Add(MoveTemp(QueuedImage));
 
 				UE_LOG(LogRemoteSession, Verbose, TEXT("finished decompressing image %d in %.02f ms (%d in queue)"),
 					Image->ImageIndex,
diff --git a/Engine/Plugins/Experimental/RemoteSession/Source/RemoteSession/Private/Channels/RemoteSessionInputChannel.cpp b/Engine/Plugins/Experimental/RemoteSession/Source/RemoteSession/Private/Channels/RemoteSessionInputChannel.cpp
--- a/Engine/Plugins/Experimental/RemoteSession/Source/RemoteSession/Private/Channels/RemoteSessionInputChannel.cpp
+++ b/Engine/Plugins/Experimental/RemoteSession/Source/RemoteSession/Private/Channels/RemoteSessionInputChannel.cpp
@@ -28,7 +28,7 @@ FRemoteSessionInputChannel::FRemoteSessionInputChannel(ERemoteSessionChannelMode
 		// and passes it through to the default UE handler
 		DefaultHandler = FSlateApplication::Get().GetPlatformApplication()->GetMessageHandler();
 
-		RecordingHandler = MakeShareable(new FRecordingMessageHandler(DefaultHandler));
+		RecordingHandler = MakeShared<FRecordingMessageHandler>(DefaultHandler);
 		RecordingHandler->SetRecordingHandler(this);
 
 		FSlateApplication::Get().GetPlatformApplication()->SetMessageHandler(RecordingHandler.ToSharedRef());
@@ -37,7 +37,7 @@ FRemoteSessionInputChannel::FRemoteSessionInputChannel(ERemoteSessionChannelMode
 	{
 		DefaultHandler = FSlateApplication::Get().GetPlatformApplication()->GetMessageHandler();
 
-		PlaybackHandler = MakeShareable(new FRecordingMessageHandler(DefaultHandler));
+		PlaybackHandler = MakeShared<FRecordingMessageHandler>(DefaultHandler);
 		
 		auto Delegate = FBackChannelDispatchDelegate::FDelegate::CreateRaw(this, &FRemoteSessionInputChannel::OnRemoteMessage);
 		MessageCallbackHandle = Connection->AddMessageHandler(TEXT("/MessageHandler/"), Delegate);
diff --git a/Engine/Plugins/Experimental/RemoteSession/Source/RemoteSession/Private/Channels/RemoteSessionXRTrackingChannel.cpp b/Engine/Plugins/Experimental/RemoteSession/Source/RemoteSession/Private/Channels/RemoteSessionXRTrackingChannel.cpp
--- a/Engine/Plugins/Experimental/RemoteSession/Source/RemoteSession/Private/Channels/RemoteSessionXRTrackingChannel.cpp
+++ b/Engine/Plugins/Experimental/RemoteSession/Source/RemoteSession/Private/Channels/RemoteSessionXRTrackingChannel.cpp
@@ -151,20 +151,24 @@ void FRemoteSessionXRTrackingChannel::ReceiveXRTracking(FBackChannelOSCMessage&
 		return;
 	}
 
-	TSharedPtr<TArray<uint8>, ESPMode::ThreadSafe> DataCopy = MakeShareable(new TArray<uint8>());
-    TSharedPtr<IXRTrackingSystem, ESPMode::ThreadSafe> TaskXRSystem = ProxyXRSystem;
+	TUniquePtr<TArray<uint8>> DataCopy = MakeUnique<TArray<uint8>>();
+    TWeakPtr<IXRTrackingSystem, ESPMode::ThreadSafe> TaskXRSystem = ProxyXRSystem;
 
 	Message << *DataCopy;
 
-    AsyncTask(ENamedThreads::GameThread, [TaskXRSystem, DataCopy]
+    AsyncTask(ENamedThreads::GameThread, [TaskXRSystem, DataCopy=MoveTemp(DataCopy)]
 	{
-		FMemoryReader Ar(*DataCopy);
-		TwoParamMsg<FVector, FRotator> MsgParam(Ar);
-        
-        UE_LOG(LogRemoteSession, Verbose, TEXT("Received Rotation (%.02f,%.02f,%.02f)"), MsgParam.Param2.Pitch,MsgParam.Param2.Yaw,MsgParam.Param2.Roll);
+		TSharedPtr<IXRTrackingSystem, ESPMode::ThreadSafe> TaskXRSystemPinned = TaskXRSystem.Pin();
+		if (TaskXRSystemPinned)
+		{
+			FMemoryReader Ar(*DataCopy);
+			TwoParamMsg<FVector, FRotator> MsgParam(Ar);
+
+			UE_LOG(LogRemoteSession, Verbose, TEXT("Received Rotation (%.02f,%.02f,%.02f)"), MsgParam.Param2.Pitch, MsgParam.Param2.Yaw, MsgParam.Param2.Roll);
 
-		FTransform NewTransform(MsgParam.Param2, MsgParam.Param1);
-        TaskXRSystem->UpdateTrackingToWorldTransform(NewTransform);
+			FTransform NewTransform(MsgParam.Param2, MsgParam.Param1);
+			TaskXRSystemPinned->UpdateTrackingToWorldTransform(NewTransform);
+		}
 	});
 }
 
diff --git a/Engine/Plugins/Experimental/RemoteSession/Source/RemoteSession/Private/ImageProviders/RemoteSessionFrameBufferImageProvider.cpp b/Engine/Plugins/Experimental/RemoteSession/Source/RemoteSession/Private/ImageProviders/RemoteSessionFrameBufferImageProvider.cpp
--- a/Engine/Plugins/Experimental/RemoteSession/Source/RemoteSession/Private/ImageProviders/RemoteSessionFrameBufferImageProvider.cpp
+++ b/Engine/Plugins/Experimental/RemoteSession/Source/RemoteSession/Private/ImageProviders/RemoteSessionFrameBufferImageProvider.cpp
@@ -32,11 +32,12 @@ static FAutoConsoleVariableRef CVarResYOverride(
 	TEXT("Sets the desired Y resolution"),
 	ECVF_Default);
 
-FRemoteSessionFrameBufferImageProvider::FRemoteSessionFrameBufferImageProvider(TWeakPtr<FRemoteSessionImageChannel> InOwner)
+FRemoteSessionFrameBufferImageProvider::FRemoteSessionFrameBufferImageProvider(TSharedPtr<FRemoteSessionImageChannel::FImageSender, ESPMode::ThreadSafe> InImageSender)
 {
-	ImageChannel = InOwner;
+	ImageSender = InImageSender;
 	LastSentImageTime = 0.0;
 	ViewportResized = false;
+	NumDecodingTasks = MakeShared<FThreadSafeCounter, ESPMode::ThreadSafe>();
 }
 
 FRemoteSessionFrameBufferImageProvider::~FRemoteSessionFrameBufferImageProvider()
@@ -112,33 +113,38 @@ void FRemoteSessionFrameBufferImageProvider::Tick(const float InDeltaTime)
 			const int32 DesiredFrameTimeMS = 1000 / FramerateMasterSetting;
 
 			// Encoding/decoding can take longer than a frame, so skip if we're still processing the previous frame
-			if (NumDecodingTasks.GetValue() == 0 && ElapsedImageTimeMS >= DesiredFrameTimeMS)
+			if (NumDecodingTasks->GetValue() == 0 && ElapsedImageTimeMS >= DesiredFrameTimeMS)
 			{
-				NumDecodingTasks.Increment();
+				NumDecodingTasks->Increment();
 
 				FCapturedFrameData& LastFrame = Frames.Last();
-
-				TArray<FColor>* ColorData = new TArray<FColor>(MoveTemp(LastFrame.ColorBuffer));
-
 				FIntPoint Size = LastFrame.BufferSize;
+				TArray<FColor> ColorData = MoveTemp(LastFrame.ColorBuffer);
+				TWeakPtr<FThreadSafeCounter, ESPMode::ThreadSafe> WeakNumDecodingTasks = NumDecodingTasks;
+				TWeakPtr<FRemoteSessionImageChannel::FImageSender, ESPMode::ThreadSafe> WeakImageSender = ImageSender;
 
-				AsyncTask(ENamedThreads::AnyBackgroundHiPriTask, [this, Size, ColorData]()
+				AsyncTask(ENamedThreads::AnyBackgroundHiPriTask, [WeakNumDecodingTasks, WeakImageSender, Size, ColorData=MoveTemp(ColorData)]() mutable
 				{
 					SCOPE_CYCLE_COUNTER(STAT_ImageCompression);
 
-					if (TSharedPtr<FRemoteSessionImageChannel> ImageChannelPinned = ImageChannel.Pin())
+					if (WeakImageSender.IsValid())
 					{
-						for (FColor& Color : *ColorData)
+						if (TSharedPtr<FThreadSafeCounter, ESPMode::ThreadSafe> NumDecodingTasksPinned = WeakNumDecodingTasks.Pin())
 						{
-							Color.A = 255;
-						}
+							for (FColor& Color : ColorData)
+							{
+								Color.A = 255;
+							}
 
-						ImageChannelPinned->SendRawImageToClients(Size.X, Size.Y, *ColorData);
-					}
+							if (TSharedPtr<FRemoteSessionImageChannel::FImageSender, ESPMode::ThreadSafe> ImageSenderPinned = WeakImageSender.Pin())
+							{
+								ImageSenderPinned->SendRawImageToClients(Size.X, Size.Y, ColorData.GetData(), ColorData.GetAllocatedSize());
+							}
 
-					delete ColorData;
+							NumDecodingTasksPinned->Decrement();
+						}
+					}
 
-					NumDecodingTasks.Decrement();
 				});
 
 				LastSentImageTime = FPlatformTime::Seconds();
diff --git a/Engine/Plugins/Experimental/RemoteSession/Source/RemoteSession/Private/ImageProviders/RemoteSessionMediaOutput.cpp b/Engine/Plugins/Experimental/RemoteSession/Source/RemoteSession/Private/ImageProviders/RemoteSessionMediaOutput.cpp
--- a/Engine/Plugins/Experimental/RemoteSession/Source/RemoteSession/Private/ImageProviders/RemoteSessionMediaOutput.cpp
+++ b/Engine/Plugins/Experimental/RemoteSession/Source/RemoteSession/Private/ImageProviders/RemoteSessionMediaOutput.cpp
@@ -39,7 +39,7 @@ void URemoteSessionMediaCapture::OnFrameCaptured_RenderingThread(const FCaptureB
 {
 	if (ImageChannel)
 	{
-		ImageChannel->SendRawImageToClients(Width, Height, InBuffer, Width * Height * sizeof(FColor));
+		ImageChannel->GetImageSender()->SendRawImageToClients(Width, Height, InBuffer, Width * Height * sizeof(FColor));
 	}
 }
 
