Apply fixes to cyclic includes from CL# 8497701
#rnx

diff --git a/Engine/Plugins/Experimental/RemoteSession/Source/RemoteSession/RemoteSession.h b/Engine/Plugins/Experimental/RemoteSession/Source/RemoteSession/RemoteSession.h
--- a/Engine/Plugins/Experimental/RemoteSession/Source/RemoteSession/RemoteSession.h
+++ b/Engine/Plugins/Experimental/RemoteSession/Source/RemoteSession/RemoteSession.h
@@ -5,7 +5,6 @@
 #include "CoreMinimal.h"
 #include "Modules/ModuleManager.h"
 #include "RemoteSessionRole.h"
-#include "Channels/RemoteSessionChannel.h"
 
 #define REMOTE_SESSION_VERSION_STRING TEXT("1.0.5")
 
