Fix CIS, load the atomic to force the compiler to read a int32.

diff --git a/Engine/Plugins/Experimental/RemoteSession/Source/RemoteSession/Private/Channels/RemoteSessionImageChannel.cpp b/Engine/Plugins/Experimental/RemoteSession/Source/RemoteSession/Private/Channels/RemoteSessionImageChannel.cpp
--- a/Engine/Plugins/Experimental/RemoteSession/Source/RemoteSession/Private/Channels/RemoteSessionImageChannel.cpp
+++ b/Engine/Plugins/Experimental/RemoteSession/Source/RemoteSession/Private/Channels/RemoteSessionImageChannel.cpp
@@ -61,8 +61,8 @@ void FRemoteSessionImageChannel::FImageSender::SendRawImageToClients(int32 Width
 
 			ImageWrapper->SetRaw(ImageData, AllocatedImageDataSize, Width, Height, ERGBFormat::BGRA, 8);
 
-			const int32 CurrentQuality = QualityMasterSetting > 0 ? QualityMasterSetting : CompressQuality;
-			const TArray<uint8>& JPGData = ImageWrapper->GetCompressed(CompressQuality);
+			const int32 CurrentQuality = QualityMasterSetting > 0 ? QualityMasterSetting : CompressQuality.Load();
+			const TArray<uint8>& JPGData = ImageWrapper->GetCompressed(CurrentQuality);
 
 			FBackChannelOSCMessage Msg(TEXT("/Screen"));
 			Msg.Write(Width);
