Fix crash when closing virtual camera

diff --git a/Engine/Plugins/Experimental/RemoteSession/Source/RemoteSession/ImageProviders/RemoteSessionFrameBufferImageProvider.h b/Engine/Plugins/Experimental/RemoteSession/Source/RemoteSession/ImageProviders/RemoteSessionFrameBufferImageProvider.h
--- a/Engine/Plugins/Experimental/RemoteSession/Source/RemoteSession/ImageProviders/RemoteSessionFrameBufferImageProvider.h
+++ b/Engine/Plugins/Experimental/RemoteSession/Source/RemoteSession/ImageProviders/RemoteSessionFrameBufferImageProvider.h
@@ -7,6 +7,7 @@
 
 class FFrameGrabber;
 class FSceneViewport;
+class SWindow;
 
 /**
  *	Use the FrameGrabber on the host to provide an image to the image channel.
@@ -30,14 +31,17 @@ public:
 	/** Signals that the viewport was resized */
 	void OnViewportResized(FVector2D NewSize);
 
-	/** Safely create the frame grabber */
-	void CreateFrameGrabber(TSharedRef<FSceneViewport> Viewport);
-
 protected:
 
-	/** Release the FrameGrabber*/
+	/** Release the FrameGrabber */
 	void ReleaseFrameGrabber();
 
+	/** When the window is destroyed */
+	void OnWindowClosedEvent(const TSharedRef<SWindow>&);
+
+	/** Safely create the frame grabber */
+	void CreateFrameGrabber(TSharedRef<FSceneViewport> Viewport);
+
 	TWeakPtr<FRemoteSessionImageChannel::FImageSender, ESPMode::ThreadSafe> ImageSender;
 
 	TSharedPtr<FFrameGrabber> FrameGrabber;
@@ -51,5 +55,5 @@ protected:
 	bool ViewportResized;
 
 	/** Holds a reference to the scene viewport */
-	TSharedPtr<FSceneViewport> SceneViewport;
+	TWeakPtr<FSceneViewport> SceneViewport;
 };
diff --git a/Engine/Plugins/Experimental/RemoteSession/Source/RemoteSession/Private/Channels/RemoteSessionInputChannel.cpp b/Engine/Plugins/Experimental/RemoteSession/Source/RemoteSession/Private/Channels/RemoteSessionInputChannel.cpp
--- a/Engine/Plugins/Experimental/RemoteSession/Source/RemoteSession/Private/Channels/RemoteSessionInputChannel.cpp
+++ b/Engine/Plugins/Experimental/RemoteSession/Source/RemoteSession/Private/Channels/RemoteSessionInputChannel.cpp
@@ -56,7 +56,7 @@ FRemoteSessionInputChannel::~FRemoteSessionInputChannel()
 	}
 
 	// todo - is this ok? Might other things have changed the handler like we do?
-	if (DefaultHandler.IsValid())
+	if (DefaultHandler.IsValid() && FSlateApplication::IsInitialized())
 	{
 		FSlateApplication::Get().GetPlatformApplication()->SetMessageHandler(DefaultHandler.ToSharedRef());
 	}
diff --git a/Engine/Plugins/Experimental/RemoteSession/Source/RemoteSession/Private/ImageProviders/RemoteSessionFrameBufferImageProvider.cpp b/Engine/Plugins/Experimental/RemoteSession/Source/RemoteSession/Private/ImageProviders/RemoteSessionFrameBufferImageProvider.cpp
--- a/Engine/Plugins/Experimental/RemoteSession/Source/RemoteSession/Private/ImageProviders/RemoteSessionFrameBufferImageProvider.cpp
+++ b/Engine/Plugins/Experimental/RemoteSession/Source/RemoteSession/Private/ImageProviders/RemoteSessionFrameBufferImageProvider.cpp
@@ -10,6 +10,7 @@
 #include "IImageWrapperModule.h"
 #include "Engine/Texture2D.h"
 #include "Modules/ModuleManager.h"
+#include "Widgets/SWindow.h"
 
 DECLARE_CYCLE_STAT(TEXT("RSFrameBufferCap"), STAT_FrameBufferCapture, STATGROUP_Game);
 DECLARE_CYCLE_STAT(TEXT("RSImageCompression"), STAT_ImageCompression, STATGROUP_Game);
@@ -65,12 +66,35 @@ void FRemoteSessionFrameBufferImageProvider::SetCaptureFrameRate(int32 InFramera
 
 void FRemoteSessionFrameBufferImageProvider::SetCaptureViewport(TSharedRef<FSceneViewport> Viewport)
 {
-	SceneViewport = Viewport;
+	if (Viewport != SceneViewport)
+	{
+		if (TSharedPtr<FSceneViewport> PreviousSceneViewportPinned = SceneViewport.Pin())
+		{
+			if (TSharedPtr<SWindow> Window = Viewport->FindWindow())
+			{
+				Window->GetOnWindowClosedEvent().RemoveAll(this);
+			}
+			PreviousSceneViewportPinned->SetOnSceneViewportResizeDel(FOnSceneViewportResize());
+		}
 
-	CreateFrameGrabber(Viewport);
+		SceneViewport = Viewport;
 
-	// set the listener for the window resize event
-	Viewport->SetOnSceneViewportResizeDel(FOnSceneViewportResize::CreateRaw(this, &FRemoteSessionFrameBufferImageProvider::OnViewportResized));
+		if (TSharedPtr<SWindow> Window = Viewport->FindWindow())
+		{
+			Window->GetOnWindowClosedEvent().AddRaw(this, &FRemoteSessionFrameBufferImageProvider::OnWindowClosedEvent);
+		}
+
+		CreateFrameGrabber(Viewport);
+
+		// set the listener for the window resize event
+		Viewport->SetOnSceneViewportResizeDel(FOnSceneViewportResize::CreateRaw(this, &FRemoteSessionFrameBufferImageProvider::OnViewportResized));
+	}
+}
+
+void FRemoteSessionFrameBufferImageProvider::OnWindowClosedEvent(const TSharedRef<SWindow>&)
+{
+	ReleaseFrameGrabber();
+	SceneViewport.Reset();
 }
 
 void FRemoteSessionFrameBufferImageProvider::CreateFrameGrabber(TSharedRef<FSceneViewport> Viewport)
@@ -78,7 +102,7 @@ void FRemoteSessionFrameBufferImageProvider::CreateFrameGrabber(TSharedRef<FScen
 	ReleaseFrameGrabber();
 
 	// For times when we want a specific resolution
-	FIntPoint FrameGrabberSize = SceneViewport->GetSize();
+	FIntPoint FrameGrabberSize = Viewport->GetSize();
 	if (FrameGrabberResX > 0)
 	{
 		FrameGrabberSize.X = FrameGrabberResX;
@@ -98,56 +122,63 @@ void FRemoteSessionFrameBufferImageProvider::Tick(const float InDeltaTime)
 	{
 		if (ViewportResized)
 		{
-			CreateFrameGrabber(SceneViewport.ToSharedRef());
+			ReleaseFrameGrabber();
+			if (TSharedPtr<FSceneViewport> SceneViewportPinned = SceneViewport.Pin())
+			{
+				CreateFrameGrabber(SceneViewportPinned.ToSharedRef());
+			}
 			ViewportResized = false;
 		}
 		SCOPE_CYCLE_COUNTER(STAT_FrameBufferCapture);
 
-		FrameGrabber->CaptureThisFrame(FFramePayloadPtr());
-
-		TArray<FCapturedFrameData> Frames = FrameGrabber->GetCapturedFrames();
-
-		if (Frames.Num())
+		if (FrameGrabber)
 		{
-			const double ElapsedImageTimeMS = (FPlatformTime::Seconds() - LastSentImageTime) * 1000;
-			const int32 DesiredFrameTimeMS = 1000 / FramerateMasterSetting;
+			FrameGrabber->CaptureThisFrame(FFramePayloadPtr());
 
-			// Encoding/decoding can take longer than a frame, so skip if we're still processing the previous frame
-			if (NumDecodingTasks->GetValue() == 0 && ElapsedImageTimeMS >= DesiredFrameTimeMS)
-			{
-				NumDecodingTasks->Increment();
+			TArray<FCapturedFrameData> Frames = FrameGrabber->GetCapturedFrames();
 
-				FCapturedFrameData& LastFrame = Frames.Last();
-				FIntPoint Size = LastFrame.BufferSize;
-				TArray<FColor> ColorData = MoveTemp(LastFrame.ColorBuffer);
-				TWeakPtr<FThreadSafeCounter, ESPMode::ThreadSafe> WeakNumDecodingTasks = NumDecodingTasks;
-				TWeakPtr<FRemoteSessionImageChannel::FImageSender, ESPMode::ThreadSafe> WeakImageSender = ImageSender;
+			if (Frames.Num())
+			{
+				const double ElapsedImageTimeMS = (FPlatformTime::Seconds() - LastSentImageTime) * 1000;
+				const int32 DesiredFrameTimeMS = 1000 / FramerateMasterSetting;
 
-				AsyncTask(ENamedThreads::AnyBackgroundHiPriTask, [WeakNumDecodingTasks, WeakImageSender, Size, ColorData=MoveTemp(ColorData)]() mutable
+				// Encoding/decoding can take longer than a frame, so skip if we're still processing the previous frame
+				if (NumDecodingTasks->GetValue() == 0 && ElapsedImageTimeMS >= DesiredFrameTimeMS)
 				{
-					SCOPE_CYCLE_COUNTER(STAT_ImageCompression);
+					NumDecodingTasks->Increment();
 
-					if (WeakImageSender.IsValid())
+					FCapturedFrameData& LastFrame = Frames.Last();
+					FIntPoint Size = LastFrame.BufferSize;
+					TArray<FColor> ColorData = MoveTemp(LastFrame.ColorBuffer);
+					TWeakPtr<FThreadSafeCounter, ESPMode::ThreadSafe> WeakNumDecodingTasks = NumDecodingTasks;
+					TWeakPtr<FRemoteSessionImageChannel::FImageSender, ESPMode::ThreadSafe> WeakImageSender = ImageSender;
+
+					AsyncTask(ENamedThreads::AnyBackgroundHiPriTask, [WeakNumDecodingTasks, WeakImageSender, Size, ColorData=MoveTemp(ColorData)]() mutable
 					{
-						if (TSharedPtr<FThreadSafeCounter, ESPMode::ThreadSafe> NumDecodingTasksPinned = WeakNumDecodingTasks.Pin())
+						SCOPE_CYCLE_COUNTER(STAT_ImageCompression);
+
+						if (WeakImageSender.IsValid())
 						{
-							for (FColor& Color : ColorData)
+							if (TSharedPtr<FThreadSafeCounter, ESPMode::ThreadSafe> NumDecodingTasksPinned = WeakNumDecodingTasks.Pin())
 							{
-								Color.A = 255;
-							}
+								for (FColor& Color : ColorData)
+								{
+									Color.A = 255;
+								}
 
-							if (TSharedPtr<FRemoteSessionImageChannel::FImageSender, ESPMode::ThreadSafe> ImageSenderPinned = WeakImageSender.Pin())
-							{
-								ImageSenderPinned->SendRawImageToClients(Size.X, Size.Y, ColorData.GetData(), ColorData.GetAllocatedSize());
-							}
+								if (TSharedPtr<FRemoteSessionImageChannel::FImageSender, ESPMode::ThreadSafe> ImageSenderPinned = WeakImageSender.Pin())
+								{
+									ImageSenderPinned->SendRawImageToClients(Size.X, Size.Y, ColorData.GetData(), ColorData.GetAllocatedSize());
+								}
 
-							NumDecodingTasksPinned->Decrement();
+								NumDecodingTasksPinned->Decrement();
+							}
 						}
-					}
 
-				});
+					});
 
-				LastSentImageTime = FPlatformTime::Seconds();
+					LastSentImageTime = FPlatformTime::Seconds();
+				}
 			}
 		}
 	}
