RemoteSession, unregister from the event when the the image provider is destroyed.

diff --git a/Engine/Plugins/Experimental/RemoteSession/Source/RemoteSession/ImageProviders/RemoteSessionFrameBufferImageProvider.h b/Engine/Plugins/Experimental/RemoteSession/Source/RemoteSession/ImageProviders/RemoteSessionFrameBufferImageProvider.h
--- a/Engine/Plugins/Experimental/RemoteSession/Source/RemoteSession/ImageProviders/RemoteSessionFrameBufferImageProvider.h
+++ b/Engine/Plugins/Experimental/RemoteSession/Source/RemoteSession/ImageProviders/RemoteSessionFrameBufferImageProvider.h
@@ -56,4 +56,7 @@ protected:
 
 	/** Holds a reference to the scene viewport */
 	TWeakPtr<FSceneViewport> SceneViewport;
+
+	/** Holds a reference to the SceneViewport SWindow */
+	TWeakPtr<SWindow> SceneViewportWindow;
 };
diff --git a/Engine/Plugins/Experimental/RemoteSession/Source/RemoteSession/Private/ImageProviders/RemoteSessionFrameBufferImageProvider.cpp b/Engine/Plugins/Experimental/RemoteSession/Source/RemoteSession/Private/ImageProviders/RemoteSessionFrameBufferImageProvider.cpp
--- a/Engine/Plugins/Experimental/RemoteSession/Source/RemoteSession/Private/ImageProviders/RemoteSessionFrameBufferImageProvider.cpp
+++ b/Engine/Plugins/Experimental/RemoteSession/Source/RemoteSession/Private/ImageProviders/RemoteSessionFrameBufferImageProvider.cpp
@@ -43,6 +43,14 @@ FRemoteSessionFrameBufferImageProvider::FRemoteSessionFrameBufferImageProvider(T
 
 FRemoteSessionFrameBufferImageProvider::~FRemoteSessionFrameBufferImageProvider()
 {
+	if (TSharedPtr<FSceneViewport> PreviousSceneViewportPinned = SceneViewport.Pin())
+	{
+		PreviousSceneViewportPinned->SetOnSceneViewportResizeDel(FOnSceneViewportResize());
+	}
+	if (TSharedPtr<SWindow> SceneViewportWindowPin = SceneViewportWindow.Pin())
+	{
+		SceneViewportWindowPin->GetOnWindowClosedEvent().RemoveAll(this);
+	}
 	ReleaseFrameGrabber();
 }
 
@@ -70,17 +78,19 @@ void FRemoteSessionFrameBufferImageProvider::SetCaptureViewport(TSharedRef<FScen
 	{
 		if (TSharedPtr<FSceneViewport> PreviousSceneViewportPinned = SceneViewport.Pin())
 		{
-			if (TSharedPtr<SWindow> Window = Viewport->FindWindow())
-			{
-				Window->GetOnWindowClosedEvent().RemoveAll(this);
-			}
 			PreviousSceneViewportPinned->SetOnSceneViewportResizeDel(FOnSceneViewportResize());
 		}
+		if (TSharedPtr<SWindow> SceneViewportWindowPin = SceneViewportWindow.Pin())
+		{
+			SceneViewportWindowPin->GetOnWindowClosedEvent().RemoveAll(this);
+		}
 
+		SceneViewportWindow.Reset();
 		SceneViewport = Viewport;
 
 		if (TSharedPtr<SWindow> Window = Viewport->FindWindow())
 		{
+			SceneViewportWindow = Window;
 			Window->GetOnWindowClosedEvent().AddRaw(this, &FRemoteSessionFrameBufferImageProvider::OnWindowClosedEvent);
 		}
 
