Select a UMG class instead of the Blueprint asset.

diff --git a/Engine/Plugins/Experimental/RemoteSession/Source/RemoteSessionEditor/Private/Widgets/SRemoteSessionStream.cpp b/Engine/Plugins/Experimental/RemoteSession/Source/RemoteSessionEditor/Private/Widgets/SRemoteSessionStream.cpp
--- a/Engine/Plugins/Experimental/RemoteSession/Source/RemoteSessionEditor/Private/Widgets/SRemoteSessionStream.cpp
+++ b/Engine/Plugins/Experimental/RemoteSession/Source/RemoteSessionEditor/Private/Widgets/SRemoteSessionStream.cpp
@@ -18,6 +18,7 @@
 
 #include "AssetData.h"
 #include "AssetRegistryModule.h"
+#include "Editor.h"
 #include "Editor/EditorEngine.h"
 #include "EditorSupportDelegates.h"
 #include "Engine/TextureRenderTarget2D.h"
@@ -26,12 +27,14 @@
 #include "LevelEditor.h"
 #include "Modules/ModuleManager.h"
 #include "PropertyEditorModule.h"
+#include "RenderingThread.h"
 #include "ScopedTransaction.h"
 #include "Stats/Stats2.h"
 #include "WidgetBlueprint.h"
 #include "WorkspaceMenuStructure.h"
 #include "WorkspaceMenuStructureModule.h"
 
+#include "RemoteSession.h"
 #include "Channels/RemoteSessionImageChannel.h"
 #include "Channels/RemoteSessionInputChannel.h"
 #include "RemoteSession/RemoteSession.h"
@@ -58,8 +61,7 @@ namespace RemoteSessionStream
 
 
 URemoteSessionStreamWidgetUserData::URemoteSessionStreamWidgetUserData()
-	: WidgetBlueprint(nullptr)
-	, Size(1280, 720)
+	: Size(1280, 720)
 	, Port(IRemoteSessionModule::kDefaultPort)
 {
 }
@@ -128,7 +130,6 @@ void SRemoteSessionStream::AddReferencedObjects(FReferenceCollector& Collector)
 	Collector.AddReferencedObject(RenderTarget2D);
 	Collector.AddReferencedObject(UserWidget);
 	Collector.AddReferencedObject(WidgetWorld);
-	Collector.AddReferencedObject(WidgetBlueprint);
 	Collector.AddReferencedObject(MediaOutput);
 	Collector.AddReferencedObject(MediaCapture);
 }
@@ -396,7 +397,7 @@ void SRemoteSessionStream::Tick(float InDeltaTime)
 		RemoteSessionHost->Tick(InDeltaTime);
 	}
 
-	if (WidgetRenderer && RenderTarget2D && VirtualWindow)
+	if (UserWidget && !UserWidget->IsDesignTime() && WidgetRenderer && RenderTarget2D && VirtualWindow)
 	{
 		WidgetRenderer->DrawWindow(RenderTarget2D, VirtualWindow->GetHittestGrid(), VirtualWindow.ToSharedRef(), 1.f, WidgetSize, InDeltaTime);
 	}
@@ -415,7 +416,7 @@ TStatId SRemoteSessionStream::GetStatId() const
 
 bool SRemoteSessionStream::CanStream() const
 {
-	return WidgetUserData && WidgetUserData->WidgetBlueprint && WidgetUserData->Size.X > 1 && WidgetUserData->Size.Y > 1;
+	return WidgetUserData && WidgetUserData->WidgetClass.Get() && WidgetUserData->Size.X > 1 && WidgetUserData->Size.Y > 1;
 }
 
 void SRemoteSessionStream::EnabledStreaming(bool bInStreaming)
@@ -426,16 +427,14 @@ void SRemoteSessionStream::EnabledStreaming(bool bInStreaming)
 		{
 			// Cache values
 			WidgetSize = WidgetUserData->Size;
-			WidgetBlueprint = WidgetUserData->WidgetBlueprint;
-			check(WidgetBlueprint);
 
 			// Create Widget
 			{
-				TSubclassOf<UUserWidget> WidgetClass = WidgetBlueprint->GeneratedClass.Get();
 				WidgetWorld = GEditor->GetEditorWorldContext().World();
-				if (WidgetClass && WidgetWorld)
+				if (WidgetWorld)
 				{
-					UserWidget = CreateWidget<UUserWidget>(WidgetWorld, WidgetClass);
+					check(WidgetUserData->WidgetClass);
+					UserWidget = CreateWidget<UUserWidget>(WidgetWorld, WidgetUserData->WidgetClass);
 				}
 			}
 
@@ -476,7 +475,13 @@ void SRemoteSessionStream::EnabledStreaming(bool bInStreaming)
 			VirtualWindow->SetContent(UserWidget->TakeWidget());
 			VirtualWindow->Resize(WidgetSize);
 
-			WidgetRenderer = MakeUnique<FWidgetRenderer>();
+			if (FSlateApplication::IsInitialized())
+			{
+				FSlateApplication::Get().RegisterVirtualWindow(VirtualWindow.ToSharedRef());
+			}
+
+			bool bApplyGammaCorrection = false;
+			WidgetRenderer = new FWidgetRenderer(bApplyGammaCorrection);
 			WidgetRenderer->DrawWindow(RenderTarget2D, VirtualWindow->GetHittestGrid(), VirtualWindow.ToSharedRef(), 1.f, WidgetSize, 0.1f);
 
 			// Register callback
@@ -486,11 +491,13 @@ void SRemoteSessionStream::EnabledStreaming(bool bInStreaming)
 			AssetRegistryModule.Get().OnAssetRemoved().AddRaw(this, &SRemoteSessionStream::HandleAssetRemoved);
 			FLevelEditorModule& LevelEditorModule = FModuleManager::LoadModuleChecked<FLevelEditorModule>(TEXT("LevelEditor"));
 			LevelEditorModule.OnMapChanged().AddRaw(this, &SRemoteSessionStream::OnMapChanged);
+			FEditorDelegates::OnAssetsCanDelete.AddRaw(this, &SRemoteSessionStream::CanDeleteAssets);
 		}
 		else
 		{
 			bInStreaming = false;
 
+			FEditorDelegates::OnAssetsCanDelete.RemoveAll(this);
 			if (FLevelEditorModule* LevelEditorModule = FModuleManager::GetModulePtr<FLevelEditorModule>(TEXT("LevelEditor")))
 			{
 				LevelEditorModule->OnMapChanged().RemoveAll(this);
@@ -502,8 +509,17 @@ void SRemoteSessionStream::EnabledStreaming(bool bInStreaming)
 			FEditorSupportDelegates::PrepareToCleanseEditorObject.RemoveAll(this);
 			GEditor->OnBlueprintPreCompile().RemoveAll(this);
 
+			if (VirtualWindow.IsValid() && FSlateApplication::IsInitialized())
+			{
+				FSlateApplication::Get().UnregisterVirtualWindow(VirtualWindow.ToSharedRef());
+			}
+
 			RenderTargetBrush.SetResourceObject(nullptr);
-			WidgetRenderer.Reset();
+			if (WidgetRenderer)
+			{
+				BeginCleanup(WidgetRenderer);
+				WidgetRenderer = nullptr;
+			}
 			VirtualWindow.Reset();
 			ResetUObject();
 			RemoteSessionHost.Reset();
@@ -519,7 +535,6 @@ void SRemoteSessionStream::ResetUObject()
 	RenderTarget2D = nullptr;
 	UserWidget = nullptr;
 	WidgetWorld = nullptr;
-	WidgetBlueprint = nullptr;
 	MediaOutput = nullptr;
 	if (MediaCapture)
 	{
@@ -552,7 +567,7 @@ void SRemoteSessionStream::OnImageChannelCreated(TWeakPtr<IRemoteSessionChannel>
 
 void SRemoteSessionStream::OnBlueprintPreCompile(UBlueprint* Blueprint)
 {
-	if (Blueprint == WidgetBlueprint)
+	if (Blueprint && UserWidget && Blueprint->GeneratedClass == UserWidget->GetClass())
 	{
 		EnabledStreaming(false);
 	}
@@ -560,7 +575,7 @@ void SRemoteSessionStream::OnBlueprintPreCompile(UBlueprint* Blueprint)
 
 void SRemoteSessionStream::OnPrepareToCleanseEditorObject(UObject* Object)
 {
-	if (Object == RenderTarget2D || Object == UserWidget || Object == WidgetWorld || Object == WidgetBlueprint || Object == MediaCapture)
+	if (Object == RenderTarget2D || Object == UserWidget || Object == WidgetWorld || (UserWidget && Object == UserWidget->GetClass()) || Object == MediaCapture)
 	{
 		EnabledStreaming(false);
 	}
@@ -568,7 +583,7 @@ void SRemoteSessionStream::OnPrepareToCleanseEditorObject(UObject* Object)
 
 void SRemoteSessionStream::HandleAssetRemoved(const FAssetData& AssetData)
 {
-	if (FAssetData(RenderTarget2D) == AssetData || FAssetData(WidgetBlueprint) == AssetData)
+	if (FAssetData(RenderTarget2D) == AssetData || (UserWidget && AssetData.GetPackage() == UserWidget->GetClass()->GetOutermost()))
 	{
 		EnabledStreaming(false);
 	}
@@ -582,5 +597,20 @@ void SRemoteSessionStream::OnMapChanged(UWorld* World, EMapChangeType ChangeType
 	}
 }
 
+void SRemoteSessionStream::CanDeleteAssets(const TArray<UObject*>& InAssetsToDelete, FCanDeleteAssetResult& CanDeleteResult)
+{
+	if (UserWidget)
+	{
+		for (UObject* Obj : InAssetsToDelete)
+		{
+			if (UserWidget->GetClass()->GetOutermost() == Obj->GetOutermost())
+			{
+				UE_LOG(LogRemoteSession, Warning, TEXT("Asset '%s' can't be deleted because it is currently used by the Remote Session Stream."), *Obj->GetPathName());
+				CanDeleteResult.Set(false);
+				break;
+			}
+		}
+	}
+}
 
 #undef LOCTEXT_NAMESPACE
diff --git a/Engine/Plugins/Experimental/RemoteSession/Source/RemoteSessionEditor/Private/Widgets/SRemoteSessionStream.h b/Engine/Plugins/Experimental/RemoteSession/Source/RemoteSessionEditor/Private/Widgets/SRemoteSessionStream.h
--- a/Engine/Plugins/Experimental/RemoteSession/Source/RemoteSessionEditor/Private/Widgets/SRemoteSessionStream.h
+++ b/Engine/Plugins/Experimental/RemoteSession/Source/RemoteSessionEditor/Private/Widgets/SRemoteSessionStream.h
@@ -19,6 +19,7 @@
 #include "SRemoteSessionStream.generated.h"
 
 struct FAssetData;
+struct FCanDeleteAssetResult;
 class FMenuBuilder;
 class FWidgetRenderer;
 class IDetailsView;
@@ -67,7 +68,7 @@ public:
 	URemoteSessionStreamWidgetUserData();
 
 	UPROPERTY(EditAnywhere, Category = "Remote Session")
-	UWidgetBlueprint* WidgetBlueprint;
+	TSubclassOf<UUserWidget> WidgetClass;
 
 	UPROPERTY(EditAnywhere, Category = "Remote Session")
 	FVector2D Size;
@@ -132,6 +133,7 @@ private:
 	void OnPrepareToCleanseEditorObject(UObject* Object);
 	void HandleAssetRemoved(const FAssetData& AssetData);
 	void OnMapChanged(UWorld* World, EMapChangeType ChangeType);
+	void CanDeleteAssets(const TArray<UObject*>& InAssetsToDelete, FCanDeleteAssetResult& CanDeleteResult);
 
 private:
 	TSharedPtr<IDetailsView> DetailView;
@@ -144,14 +146,13 @@ private:
 	UTextureRenderTarget2D* RenderTarget2D;
 	UUserWidget* UserWidget;
 	UWorld* WidgetWorld;
-	UWidgetBlueprint* WidgetBlueprint;
 	URemoteSessionMediaOutput* MediaOutput;
 	URemoteSessionMediaCapture* MediaCapture;
 	//~ End GC by AddReferencedObjects
 
 	TSharedPtr<IRemoteSessionUnmanagedRole> RemoteSessionHost;
 	TSharedPtr<SVirtualWindow> VirtualWindow;
-	TUniquePtr<FWidgetRenderer> WidgetRenderer;
+	FWidgetRenderer* WidgetRenderer;
 	TWeakPtr<ILevelViewport> ActiveLevelViewport;
 	FVector2D WidgetSize;
 
