Remove deprecated ILevelViewport interface.

#rnx

diff --git a/Engine/Plugins/Cameras/CameraShakePreviewer/Source/CameraShakePreviewer/Private/CameraShakePreviewerModule.cpp b/Engine/Plugins/Cameras/CameraShakePreviewer/Source/CameraShakePreviewer/Private/CameraShakePreviewerModule.cpp
--- a/Engine/Plugins/Cameras/CameraShakePreviewer/Source/CameraShakePreviewer/Private/CameraShakePreviewerModule.cpp
+++ b/Engine/Plugins/Cameras/CameraShakePreviewer/Source/CameraShakePreviewer/Private/CameraShakePreviewerModule.cpp
@@ -5,7 +5,6 @@
 #include "Framework/MultiBox/MultiBoxBuilder.h"
 #include "Framework/MultiBox/MultiBoxExtender.h"
 #include "ILevelEditor.h"
-#include "ILevelViewport.h"
 #include "LevelEditor.h"
 #include "LevelEditorViewport.h"
 #include "Modules/ModuleManager.h"
diff --git a/Engine/Plugins/Experimental/RemoteSession/Source/RemoteSessionEditor/Private/Widgets/SRemoteSessionStream.cpp b/Engine/Plugins/Experimental/RemoteSession/Source/RemoteSessionEditor/Private/Widgets/SRemoteSessionStream.cpp
--- a/Engine/Plugins/Experimental/RemoteSession/Source/RemoteSessionEditor/Private/Widgets/SRemoteSessionStream.cpp
+++ b/Engine/Plugins/Experimental/RemoteSession/Source/RemoteSessionEditor/Private/Widgets/SRemoteSessionStream.cpp
@@ -23,7 +23,6 @@
 #include "EditorSupportDelegates.h"
 #include "Engine/TextureRenderTarget2D.h"
 #include "IDetailsView.h"
-#include "ILevelViewport.h"
 #include "LevelEditor.h"
 #include "Modules/ModuleManager.h"
 #include "PropertyEditorModule.h"
diff --git a/Engine/Plugins/Experimental/RemoteSession/Source/RemoteSessionEditor/Private/Widgets/SRemoteSessionStream.h b/Engine/Plugins/Experimental/RemoteSession/Source/RemoteSessionEditor/Private/Widgets/SRemoteSessionStream.h
--- a/Engine/Plugins/Experimental/RemoteSession/Source/RemoteSessionEditor/Private/Widgets/SRemoteSessionStream.h
+++ b/Engine/Plugins/Experimental/RemoteSession/Source/RemoteSessionEditor/Private/Widgets/SRemoteSessionStream.h
@@ -23,7 +23,6 @@ struct FCanDeleteAssetResult;
 class FMenuBuilder;
 class FWidgetRenderer;
 class IDetailsView;
-class ILevelViewport;
 class SImage;
 class SSplitter;
 class SVirtualWindow;
@@ -153,7 +152,6 @@ private:
 	TSharedPtr<IRemoteSessionUnmanagedRole> RemoteSessionHost;
 	TSharedPtr<SVirtualWindow> VirtualWindow;
 	FWidgetRenderer* WidgetRenderer;
-	TWeakPtr<ILevelViewport> ActiveLevelViewport;
 	FVector2D WidgetSize;
 
 	bool bIsStreaming;
diff --git a/Engine/Source/Editor/LevelEditor/Public/ILevelViewport.h b/Engine/Source/Editor/LevelEditor/Public/ILevelViewport.h
deleted file mode 100644
