Fixed the reading of OSC data for touch & motion events. LWC changed a FVector2D to doubles, so deserializing the bytestream resulting in incorrect value, effectively breaking user interaction from Live Link VCAM iOS. The reads are now forced to use FVector2f (FVector3f in motion events) and then converted to LWC for the rest of the pipeline.

diff --git a/Engine/Plugins/Experimental/RemoteSession/Source/RemoteSession/Private/MessageHandler/RecordingMessageHandler.cpp b/Engine/Plugins/Experimental/RemoteSession/Source/RemoteSession/Private/MessageHandler/RecordingMessageHandler.cpp
--- a/Engine/Plugins/Experimental/RemoteSession/Source/RemoteSession/Private/MessageHandler/RecordingMessageHandler.cpp
+++ b/Engine/Plugins/Experimental/RemoteSession/Source/RemoteSession/Private/MessageHandler/RecordingMessageHandler.cpp
@@ -92,8 +92,10 @@ void FRecordingMessageHandler::SetInputRect(const FVector2D& TopLeft, const FVec
 	InputRect = FRect(TopLeft, Extents);
 }
 
-bool FRecordingMessageHandler::ConvertToNormalizedScreenLocation(const FVector2D& InLocation, FVector2D& OutLocation)
+bool FRecordingMessageHandler::ConvertToNormalizedScreenLocation(const FVector2D& InLocation, FVector2f& OutLocation)
 {
+	// note : following the LWC update, we do a FVector2D->FVector2f conversion here.
+
 	FRect ClipRect = InputRect;
 	FIntPoint Point = FIntPoint((int)InLocation.X, (int)InLocation.Y);
 
@@ -104,18 +106,21 @@ bool FRecordingMessageHandler::ConvertToNormalizedScreenLocation(const FVector2D
 
 	if (!ClipRect.Contains(Point))
 	{
-		OutLocation = FVector2D(EForceInit::ForceInitToZero);
+		OutLocation = FVector2f(EForceInit::ForceInitToZero);
 		return false;
 	}
 
-	OutLocation = FVector2D((InLocation.X-ClipRect.X) / ClipRect.Width, (InLocation.Y - ClipRect.Y) / ClipRect.Height);
+	OutLocation = FVector2f((InLocation.X-ClipRect.X) / ClipRect.Width, (InLocation.Y - ClipRect.Y) / ClipRect.Height);
 
 	return true;
 }
 
-FVector2D FRecordingMessageHandler::ConvertFromNormalizedScreenLocation(const FVector2D& ScreenLocation)
+FVector2D FRecordingMessageHandler::ConvertFromNormalizedScreenLocation(const FVector2f& ScreenLocation)
 {
-	FVector2D OutVector = ScreenLocation;
+	// note : following the LWC update, we do a FVector2f->FVector2D conversion here.
+
+	FVector2D ScreenLocation2D = FVector2D(ScreenLocation.X, ScreenLocation.Y);
+	FVector2D OutVector = ScreenLocation2D;
 
 	TSharedPtr<SWindow> GameWindow = PlaybackWindow.Pin();
 	TSharedPtr<FSceneViewport> GameWidget = PlaybackViewport.Pin();
@@ -142,7 +147,7 @@ FVector2D FRecordingMessageHandler::ConvertFromNormalizedScreenLocation(const FV
 					FVector2D WindowClientOffset = ArrangedWidget.Geometry.GetAbsolutePosition();
 					FVector2D WindowClientSize = ArrangedWidget.Geometry.GetAbsoluteSize();
 
-					OutVector = WindowOrigin + WindowClientOffset + (ScreenLocation * WindowClientSize);
+					OutVector = WindowOrigin + WindowClientOffset + (ScreenLocation2D * WindowClientSize);
 
 				}
 			}
@@ -150,7 +155,7 @@ FVector2D FRecordingMessageHandler::ConvertFromNormalizedScreenLocation(const FV
 		else
 		{
 			FVector2D SizeInScreen = GameWindow->GetSizeInScreen();
-			OutVector = SizeInScreen * ScreenLocation;
+			OutVector = SizeInScreen * ScreenLocation2D;
 			//OutVector = GameWindow->GetLocalToScreenTransform().TransformPoint(ScreenLocation);
 		}
 	}
@@ -271,12 +276,12 @@ bool FRecordingMessageHandler::OnTouchStarted(const TSharedPtr< FGenericWindow >
 {
 	if (IsRecording())
 	{
-		FVector2D Normalized;
+		FVector2f Normalized;
 
 		if (ConvertToNormalizedScreenLocation(Location, Normalized))
 		{
 			// note - force is serialized last for backwards compat - force was introduced in 4.20
-			FourParamMsg<FVector2D, int32, int32, float> Msg(Normalized, TouchIndex, ControllerId, Force);
+			FourParamMsg<FVector2f, int32, int32, float> Msg(Normalized, TouchIndex, ControllerId, Force);
 			RecordMessage(TEXT("OnTouchStarted"), Msg.AsData());
 		}
 	}
@@ -314,7 +319,9 @@ bool FRecordingMessageHandler::OnTouchStarted(const TSharedPtr< FGenericWindow >
 
 void FRecordingMessageHandler::PlayOnTouchStarted(FArchive& Ar)
 {
-	FourParamMsg<FVector2D, int32, int32, float > Msg(Ar);
+	// LWC FVector2D are now doubles, but the OSC clients will sent floats. So we deserialize as floats
+	// then create LWC version via ConvertFromNormalizedScreenLocation
+	FourParamMsg<FVector2f, int32, int32, float > Msg(Ar);
 	FVector2D ScreenLocation = ConvertFromNormalizedScreenLocation(Msg.Param1);
 
 	TSharedPtr<FGenericWindow> Window;
@@ -332,12 +339,12 @@ bool FRecordingMessageHandler::OnTouchMoved(const FVector2D& Location, float For
 {
 	if (IsRecording())
 	{
-		FVector2D Normalized;
+		FVector2f Normalized;
 
 		if (ConvertToNormalizedScreenLocation(Location, Normalized))
 		{
 			// note - force is serialized last for backwards compat - force was introduced in 4.20
-			FourParamMsg<FVector2D, int32, int32, float> Msg(Normalized, TouchIndex, ControllerId, Force);
+			FourParamMsg<FVector2f, int32, int32, float> Msg(Normalized, TouchIndex, ControllerId, Force);
 			OutputWriter->RecordMessage(TEXT("OnTouchMoved"), Msg.AsData());
 		}
 	}
@@ -372,7 +379,9 @@ bool FRecordingMessageHandler::OnTouchMoved(const FVector2D& Location, float For
 
 void FRecordingMessageHandler::PlayOnTouchMoved(FArchive& Ar)
 {
-	FourParamMsg<FVector2D, int32, int32, float > Msg(Ar);
+	// LWC FVector2D are now doubles, but the OSC clients will sent floats. So we deserialize as floats
+	// then create LWC version via ConvertFromNormalizedScreenLocation
+	FourParamMsg<FVector2f, int32, int32, float > Msg(Ar);
 	FVector2D ScreenLocation = ConvertFromNormalizedScreenLocation(Msg.Param1);
 	// note - force is serialized last for backwards compat - force was introduced in 4.20
 	OnTouchMoved(ScreenLocation, Msg.Param4, Msg.Param2, Msg.Param3);
@@ -382,7 +391,7 @@ bool FRecordingMessageHandler::OnTouchEnded(const FVector2D& Location, int32 Tou
 {
 	if (IsRecording())
 	{
-		FVector2D Normalized;
+		FVector2f Normalized;
 
 		// if outside our bounds, end the touch where it left
 		if (ConvertToNormalizedScreenLocation(Location, Normalized) == false)
@@ -390,7 +399,7 @@ bool FRecordingMessageHandler::OnTouchEnded(const FVector2D& Location, int32 Tou
 			ConvertToNormalizedScreenLocation(LastTouchLocation, Normalized);
 		}
 		
-		ThreeParamMsg<FVector2D, int32, int32> Msg(Normalized, TouchIndex, ControllerId);
+		ThreeParamMsg<FVector2f, int32, int32> Msg(Normalized, TouchIndex, ControllerId);
 		OutputWriter->RecordMessage(TEXT("OnTouchEnded"), Msg.AsData());
 	}
 
@@ -423,7 +432,9 @@ bool FRecordingMessageHandler::OnTouchEnded(const FVector2D& Location, int32 Tou
 
 void FRecordingMessageHandler::PlayOnTouchEnded(FArchive& Ar)
 {
-	ThreeParamMsg<FVector2D, int32, int32 > Msg(Ar);
+	// LWC FVector2D are now doubles, but the OSC clients will sent floats. So we deserialize as floats
+	// then create LWC version via ConvertFromNormalizedScreenLocation
+	ThreeParamMsg<FVector2f, int32, int32 > Msg(Ar);
 	FVector2D ScreenLocation = ConvertFromNormalizedScreenLocation(Msg.Param1);
 	OnTouchEnded(ScreenLocation, Msg.Param2, Msg.Param3);
 }
@@ -432,12 +443,12 @@ bool FRecordingMessageHandler::OnTouchForceChanged(const FVector2D& Location, fl
 {
 	if (IsRecording())
 	{
-		FVector2D Normalized;
+		FVector2f Normalized;
 
 		if (ConvertToNormalizedScreenLocation(Location, Normalized))
 		{
 			// note - force is serialized last for backwards compat - force was introduced in 4.20
-			FourParamMsg<FVector2D, int32, int32, float> Msg(Normalized, TouchIndex, ControllerId, Force);
+			FourParamMsg<FVector2f, int32, int32, float> Msg(Normalized, TouchIndex, ControllerId, Force);
 			OutputWriter->RecordMessage(TEXT("OnTouchForceChanged"), Msg.AsData());
 		}
 	}
@@ -473,7 +484,9 @@ bool FRecordingMessageHandler::OnTouchForceChanged(const FVector2D& Location, fl
 
 void FRecordingMessageHandler::PlayOnTouchForceChanged(FArchive& Ar)
 {
-	FourParamMsg<FVector2D, int32, int32, float > Msg(Ar);
+	// LWC FVector2D are now doubles, but the OSC clients will sent floats. So we deserialize as floats
+	// then create LWC version via ConvertFromNormalizedScreenLocation
+	FourParamMsg<FVector2f, int32, int32, float > Msg(Ar);
 	FVector2D ScreenLocation = ConvertFromNormalizedScreenLocation(Msg.Param1);
 	OnTouchForceChanged(ScreenLocation, Msg.Param4, Msg.Param2, Msg.Param3);
 }
@@ -482,12 +495,12 @@ bool FRecordingMessageHandler::OnTouchFirstMove(const FVector2D& Location, float
 {
 	if (IsRecording())
 	{
-		FVector2D Normalized;
+		FVector2f Normalized;
 
 		if (ConvertToNormalizedScreenLocation(Location, Normalized))
 		{
 			// note - force is serialized last for backwards compat - force was introduced in 4.20
-			FourParamMsg<FVector2D, int32, int32, float> Msg(Normalized, TouchIndex, ControllerId, Force);
+			FourParamMsg<FVector2f, int32, int32, float> Msg(Normalized, TouchIndex, ControllerId, Force);
 			OutputWriter->RecordMessage(TEXT("OnTouchFirstMove"), Msg.AsData());
 		}
 	}
@@ -522,7 +535,9 @@ bool FRecordingMessageHandler::OnTouchFirstMove(const FVector2D& Location, float
 
 void FRecordingMessageHandler::PlayOnTouchFirstMove(FArchive& Ar)
 {
-	FourParamMsg<FVector2D, int32, int32, float > Msg(Ar);
+	// LWC FVector2D are now doubles, but the OSC clients will sent floats. So we deserialize as floats
+	// then create the LWC version via ConvertFromNormalizedScreenLocation
+	FourParamMsg<FVector2f, int32, int32, float > Msg(Ar);
 	FVector2D ScreenLocation = ConvertFromNormalizedScreenLocation(Msg.Param1);
 	// note - force is serialized last for backwards compat - force was introduced in 4.20
 	OnTouchFirstMove(ScreenLocation, Msg.Param4, Msg.Param2, Msg.Param3);
@@ -568,8 +583,8 @@ bool FRecordingMessageHandler::OnTouchGesture(EGestureEvent GestureType, const F
 
 void FRecordingMessageHandler::PlayOnTouchGesture(FArchive& Ar)
 {
-	FourParamMsg<uint32, FVector2D, float, bool> Msg(Ar);
-	OnTouchGesture((EGestureEvent)Msg.Param1, Msg.Param2, Msg.Param3, Msg.Param4);
+	FourParamMsg<uint32, FVector2f, float, bool> Msg(Ar);
+	OnTouchGesture((EGestureEvent)Msg.Param1, FVector2D(Msg.Param2.X, Msg.Param2.Y), Msg.Param3, Msg.Param4);
 }
 
 void FRecordingMessageHandler::OnEndGesture()
@@ -613,8 +628,15 @@ bool FRecordingMessageHandler::OnMotionDetected(const FVector& Tilt, const FVect
 
 void FRecordingMessageHandler::PlayOnMotionDetected(FArchive& Ar)
 {
-	FiveParamMsg<FVector, FVector, FVector, FVector, int32 > Msg(Ar);
-	OnMotionDetected(Msg.Param1, Msg.Param2, Msg.Param3, Msg.Param4, Msg.Param5);
+	// LWC FVector2D are now doubles, but the OSC clients will sent floats. So we deserialize as floats
+	// then create LWC versions below
+	FiveParamMsg<FVector3f, FVector3f, FVector3f, FVector3f, int32 > Msg(Ar);
+	OnMotionDetected(
+		FVector(Msg.Param1.X, Msg.Param1.Y, Msg.Param1.Z), 
+		FVector(Msg.Param2.X, Msg.Param2.Y, Msg.Param2.Z),
+		FVector(Msg.Param3.X, Msg.Param3.Y, Msg.Param3.Z),
+		FVector(Msg.Param4.X, Msg.Param4.Y, Msg.Param4.Z),
+		Msg.Param5);
 }
 
 bool FRecordingMessageHandler::OnControllerAnalog(FGamepadKeyNames::Type KeyName, int32 ControllerId, float AnalogValue)
diff --git a/Engine/Plugins/Experimental/RemoteSession/Source/RemoteSession/Private/MessageHandler/RecordingMessageHandler.h b/Engine/Plugins/Experimental/RemoteSession/Source/RemoteSession/Private/MessageHandler/RecordingMessageHandler.h
--- a/Engine/Plugins/Experimental/RemoteSession/Source/RemoteSession/Private/MessageHandler/RecordingMessageHandler.h
+++ b/Engine/Plugins/Experimental/RemoteSession/Source/RemoteSession/Private/MessageHandler/RecordingMessageHandler.h
@@ -131,8 +131,8 @@ public:
 
 protected:
 
-	bool ConvertToNormalizedScreenLocation(const FVector2D& InLocation, FVector2D& OutLocation);
-	FVector2D ConvertFromNormalizedScreenLocation(const FVector2D& ScreenLocation);
+	bool ConvertToNormalizedScreenLocation(const FVector2D& InLocation, FVector2f& OutLocation);
+	FVector2D ConvertFromNormalizedScreenLocation(const FVector2f& ScreenLocation);
 
 	void RecordMessage(const TCHAR* MsgName, const TArray<uint8>& Data);
 
