Fixed the RemoteSession message handler proxy calling the wrong function.

PR #9258: Fix copy-paste error in FProxyMessageHandler (Contributed by [REDACTED_PERSON])

diff --git a/Engine/Plugins/Experimental/RemoteSession/Source/RemoteSession/Private/MessageHandler/ProxyMessageHandler.cpp b/Engine/Plugins/Experimental/RemoteSession/Source/RemoteSession/Private/MessageHandler/ProxyMessageHandler.cpp
--- a/Engine/Plugins/Experimental/RemoteSession/Source/RemoteSession/Private/MessageHandler/ProxyMessageHandler.cpp
+++ b/Engine/Plugins/Experimental/RemoteSession/Source/RemoteSession/Private/MessageHandler/ProxyMessageHandler.cpp
@@ -358,7 +358,7 @@ void FProxyMessageHandler::OnMovedWindow(const TSharedRef< FGenericWindow >& Win
 {
 	if (TargetHandler.IsValid())
 	{
-		return TargetHandler->HandleDPIScaleChanged(Window);
+		return TargetHandler->OnMovedWindow(Window, X, Y);
 	}
 }
 
@@ -476,4 +476,4 @@ bool FProxyMessageHandler::OnWindowAction(const TSharedRef< FGenericWindow >& Wi
 	}
 
 	return true;
-}
\ No newline at end of file
+}
