Remove uses of allowlist and blacklist per inclusive language standard.

diff --git a/Engine/Plugins/Experimental/RemoteSession/Config/BaseRemoteSession.ini b/Engine/Plugins/Experimental/RemoteSession/Config/BaseRemoteSession.ini
new file mode 100644
--- /dev/null
+++ b/Engine/Plugins/Experimental/RemoteSession/Config/BaseRemoteSession.ini
@@ -0,0 +1,3 @@
+[CoreRedirects]
++PropertyRedirects=(OldName="RemoteSessionSettings.WhitelistedChannels",NewName="AllowedChannels")
++PropertyRedirects=(OldName="RemoteSessionSettings.BlacklistedChannels",NewName="DeniedChannels")
diff --git a/Engine/Plugins/Experimental/RemoteSession/README.md b/Engine/Plugins/Experimental/RemoteSession/README.md
--- a/Engine/Plugins/Experimental/RemoteSession/README.md
+++ b/Engine/Plugins/Experimental/RemoteSession/README.md
@@ -28,9 +28,9 @@ The values shown below are the current defaults. If you do not want to change fr
 	; Framerate of from the host (will be min of this and the actual game)
 	FrameRate = 60
 	; Restrict to these channels. If empty all registered channels are available
-	+WhitelistedChannels="ChannelName"
+	+AllowedChannels="ChannelName"
 	; Don't allow these channels to be used
-	+BlacklistedChannels="ChannelName"
+	+DeniedChannels="ChannelName"
 
 
 Framerate and Quality can be adjusted at runtime via the remote.framerate and remote.quality cvars.
diff --git a/Engine/Plugins/Experimental/RemoteSession/Source/RemoteSession/Private/RemoteSessionModule.cpp b/Engine/Plugins/Experimental/RemoteSession/Source/RemoteSession/Private/RemoteSessionModule.cpp
--- a/Engine/Plugins/Experimental/RemoteSession/Source/RemoteSession/Private/RemoteSessionModule.cpp
+++ b/Engine/Plugins/Experimental/RemoteSession/Source/RemoteSession/Private/RemoteSessionModule.cpp
@@ -157,23 +157,23 @@ bool FRemoteSessionModule::HandleSettingsSaved()
 	TArray<FRemoteSessionChannelInfo> KnownChannels = FRemoteSessionChannelRegistry::Get().GetRegisteredFactories();
 
 	// check channels
-	for (const auto& Channel : Settings->WhitelistedChannels)
+	for (const auto& Channel : Settings->AllowedChannels)
 	{
 		if (KnownChannels.ContainsByPredicate([Channel](const FRemoteSessionChannelInfo& Info) {
 			return Info.Type == Channel;
 			}))
 		{
-			UE_LOG(LogRemoteSession, Error, TEXT("Channel %s in the ini file whitelist is not a recognized channel."), *Channel);
+			UE_LOG(LogRemoteSession, Error, TEXT("Channel %s in the ini file allow list is not a recognized channel."), *Channel);
 		}
 	}
 
-	for (const auto& Channel : Settings->BlacklistedChannels)
+	for (const auto& Channel : Settings->DeniedChannels)
 	{
 		if (KnownChannels.ContainsByPredicate([Channel](const FRemoteSessionChannelInfo& Info) {
 			return Info.Type == Channel;
 			}))
 		{
-			UE_LOG(LogRemoteSession, Error, TEXT("Channel %s in the ini file blacklist is not a recognized channel."), *Channel);
+			UE_LOG(LogRemoteSession, Error, TEXT("Channel %s in the ini file deny list is not a recognized channel."), *Channel);
 		}
 	}
 
@@ -244,20 +244,20 @@ void FRemoteSessionModule::InitHost(const int16 Port /*= 0*/)
 
 	SupportedChannels = FRemoteSessionChannelRegistry::Get().GetRegisteredFactories();
 
-	if (Settings->WhitelistedChannels.Num())
+	if (Settings->AllowedChannels.Num())
 	{
-		// Remove anything not in the whitelist
+		// Remove anything not in the allow list
 		SupportedChannels = SupportedChannels.FilterByPredicate([Settings](const FRemoteSessionChannelInfo& Info) {
-			return Settings->WhitelistedChannels.Contains(Info.Type);
+			return Settings->AllowedChannels.Contains(Info.Type);
 		});
 	}
 
 
-	if (Settings->BlacklistedChannels.Num())
+	if (Settings->DeniedChannels.Num())
 	{
-		// Remove anything in the blacklist
+		// Remove anything in the denied list
 		SupportedChannels = SupportedChannels.FilterByPredicate([Settings](const FRemoteSessionChannelInfo& Info) {
-			return Settings->BlacklistedChannels.Contains(Info.Type) == false;
+			return Settings->DeniedChannels.Contains(Info.Type) == false;
 		});
 	}
 
diff --git a/Engine/Plugins/Experimental/RemoteSession/Source/RemoteSession/Public/RemoteSessionTypes.h b/Engine/Plugins/Experimental/RemoteSession/Source/RemoteSession/Public/RemoteSessionTypes.h
--- a/Engine/Plugins/Experimental/RemoteSession/Source/RemoteSession/Public/RemoteSessionTypes.h
+++ b/Engine/Plugins/Experimental/RemoteSession/Source/RemoteSession/Public/RemoteSessionTypes.h
@@ -105,9 +105,9 @@ public:
 
 	/* Restrict remote session to these channels. If empty all registered channels are available */
 	UPROPERTY(config)
-	TArray<FString> WhitelistedChannels;
+	TArray<FString> AllowedChannels;
 
 	/* Don't allow these channels to be used */
 	UPROPERTY(config)
-	TArray<FString> BlacklistedChannels;
-};
\ No newline at end of file
+	TArray<FString> DeniedChannels;
+};
