Fix a silent inlined depenceny in the RemoteSession and Bridge plugins.

They both implement their own MessageHandler, but don't depend on ApplicationCore. The only reason that this would compile before was because of the base GenericApplicationMessageHandler being inlined into their files. If we change anything in the GenericMessageHandler then we get linker errors. We will need to reference the InputDeviceMapper in there, so this is a necessary change.

#rnx

diff --git a/Engine/Plugins/Bridge/Source/Bridge/Bridge.Build.cs b/Engine/Plugins/Bridge/Source/Bridge/Bridge.Build.cs
--- a/Engine/Plugins/Bridge/Source/Bridge/Bridge.Build.cs
+++ b/Engine/Plugins/Bridge/Source/Bridge/Bridge.Build.cs
@@ -44,6 +44,7 @@ namespace UnrealBuildTool.Rules
 					"ContentBrowserData",
 					"PlacementMode",
 					"MegascansPlugin",
+					"ApplicationCore",
 					});
 
 			PrivateIncludePaths.AddRange(
diff --git a/Engine/Plugins/Experimental/RemoteSession/Source/RemoteSession/RemoteSession.Build.cs b/Engine/Plugins/Experimental/RemoteSession/Source/RemoteSession/RemoteSession.Build.cs
--- a/Engine/Plugins/Experimental/RemoteSession/Source/RemoteSession/RemoteSession.Build.cs
+++ b/Engine/Plugins/Experimental/RemoteSession/Source/RemoteSession/RemoteSession.Build.cs
@@ -21,6 +21,7 @@ public class RemoteSession : ModuleRules
 				"Core",
 				"MediaIOCore",
 				"BackChannel",
+				"ApplicationCore",
 				// ... add other public dependencies that you statically link with here ...
 			}
 		);
