Udpate the ProxyMessageHandler and the RecordingMessageHandler to have the newer version of OnController functoins that take in a FPlatformUser instead of an int32 ControllerId

diff --git a/Engine/Plugins/Experimental/RemoteSession/Source/RemoteSession/Private/MessageHandler/ProxyMessageHandler.cpp b/Engine/Plugins/Experimental/RemoteSession/Source/RemoteSession/Private/MessageHandler/ProxyMessageHandler.cpp
--- a/Engine/Plugins/Experimental/RemoteSession/Source/RemoteSession/Private/MessageHandler/ProxyMessageHandler.cpp
+++ b/Engine/Plugins/Experimental/RemoteSession/Source/RemoteSession/Private/MessageHandler/ProxyMessageHandler.cpp
@@ -173,6 +173,7 @@ bool FProxyMessageHandler::OnCursorSet()
 	return false;
 }
 
+PRAGMA_DISABLE_DEPRECATION_WARNINGS
 bool FProxyMessageHandler::OnControllerAnalog(FGamepadKeyNames::Type KeyName, int32 ControllerId, float AnalogValue)
 {
 	if (TargetHandler.IsValid())
@@ -202,6 +203,38 @@ bool FProxyMessageHandler::OnControllerButtonReleased(FGamepadKeyNames::Type Key
 
 	return false;
 }
+PRAGMA_ENABLE_DEPRECATION_WARNINGS
+
+bool FProxyMessageHandler::OnControllerAnalog(FGamepadKeyNames::Type KeyName, FPlatformUserId PlatformUserId, FInputDeviceId InputDeviceId, float AnalogValue)
+{
+	if (TargetHandler.IsValid())
+	{
+		return TargetHandler->OnControllerAnalog(KeyName, PlatformUserId, InputDeviceId, AnalogValue);
+	}
+
+	return false;
+}
+
+bool FProxyMessageHandler::OnControllerButtonPressed(FGamepadKeyNames::Type KeyName, FPlatformUserId PlatformUserId, FInputDeviceId InputDeviceId, bool IsRepeat)
+{
+	if (TargetHandler.IsValid())
+	{
+		return TargetHandler->OnControllerButtonPressed(KeyName, PlatformUserId, InputDeviceId, IsRepeat);
+	}
+
+	return false;
+}
+
+bool FProxyMessageHandler::OnControllerButtonReleased(FGamepadKeyNames::Type KeyName, FPlatformUserId PlatformUserId, FInputDeviceId InputDeviceId, bool IsRepeat)
+{
+	if (TargetHandler.IsValid())
+	{
+		return TargetHandler->OnControllerButtonReleased(KeyName, PlatformUserId, InputDeviceId, IsRepeat);
+	}
+
+	return false;
+}
+
 
 void FProxyMessageHandler::OnBeginGesture()
 {
diff --git a/Engine/Plugins/Experimental/RemoteSession/Source/RemoteSession/Private/MessageHandler/ProxyMessageHandler.h b/Engine/Plugins/Experimental/RemoteSession/Source/RemoteSession/Private/MessageHandler/ProxyMessageHandler.h
--- a/Engine/Plugins/Experimental/RemoteSession/Source/RemoteSession/Private/MessageHandler/ProxyMessageHandler.h
+++ b/Engine/Plugins/Experimental/RemoteSession/Source/RemoteSession/Private/MessageHandler/ProxyMessageHandler.h
@@ -53,12 +53,21 @@ public:
 
 	virtual bool OnCursorSet() override;
 
+	UE_DEPRECATED(5.1, "This version of OnControllerAnalog has been deprecated, please use the one that takes an FPlatformUser and FInputDeviceId instead.")
 	virtual bool OnControllerAnalog(FGamepadKeyNames::Type KeyName, int32 ControllerId, float AnalogValue) override;
 
+	UE_DEPRECATED(5.1, "This version of OnControllerButtonPressed has been deprecated, please use the one that takes an FPlatformUser and FInputDeviceId instead.")
 	virtual bool OnControllerButtonPressed(FGamepadKeyNames::Type KeyName, int32 ControllerId, bool IsRepeat) override;
 
+	UE_DEPRECATED(5.1, "This version of OnControllerButtonReleased has been deprecated, please use the one that takes an FPlatformUser and FInputDeviceId instead.")
 	virtual bool OnControllerButtonReleased(FGamepadKeyNames::Type KeyName, int32 ControllerId, bool IsRepeat) override;
 
+	virtual bool OnControllerAnalog(FGamepadKeyNames::Type KeyName, FPlatformUserId PlatformUserId, FInputDeviceId InputDeviceId, float AnalogValue) override;
+
+	virtual bool OnControllerButtonPressed(FGamepadKeyNames::Type KeyName, FPlatformUserId PlatformUserId, FInputDeviceId InputDeviceId, bool IsRepeat) override;
+
+	virtual bool OnControllerButtonReleased(FGamepadKeyNames::Type KeyName, FPlatformUserId PlatformUserId, FInputDeviceId InputDeviceId, bool IsRepeat) override;
+
 	virtual void OnBeginGesture() override;
 
 	virtual bool OnTouchGesture(EGestureEvent GestureType, const FVector2D& Delta, float WheelDelta, bool bIsDirectionInvertedFromDevice) override;
diff --git a/Engine/Plugins/Experimental/RemoteSession/Source/RemoteSession/Private/MessageHandler/RecordingMessageHandler.cpp b/Engine/Plugins/Experimental/RemoteSession/Source/RemoteSession/Private/MessageHandler/RecordingMessageHandler.cpp
--- a/Engine/Plugins/Experimental/RemoteSession/Source/RemoteSession/Private/MessageHandler/RecordingMessageHandler.cpp
+++ b/Engine/Plugins/Experimental/RemoteSession/Source/RemoteSession/Private/MessageHandler/RecordingMessageHandler.cpp
@@ -61,9 +61,15 @@ FRecordingMessageHandler::FRecordingMessageHandler(const TSharedPtr<FGenericAppl
 	BIND_PLAYBACK_HANDLER(TEXT("OnEndGesture"), PlayOnEndGesture);
 	BIND_PLAYBACK_HANDLER(TEXT("OnTouchForceChanged"), PlayOnTouchForceChanged);
 
+	/** Deprecated controller handlers that take in the old int32 ControllerId */
 	BIND_PLAYBACK_HANDLER(TEXT("OnControllerAnalog"), PlayOnControllerAnalog);
 	BIND_PLAYBACK_HANDLER(TEXT("OnControllerButtonPressed"), PlayOnControllerButtonPressed);
 	BIND_PLAYBACK_HANDLER(TEXT("OnControllerButtonReleased"), PlayOnControllerButtonReleased);
+
+	/** New controller handlers that take in FPlatformUserId and FInputDeviceId */
+	BIND_PLAYBACK_HANDLER(TEXT("OnControllerAnalogWithPlatformUser"), PlayOnControllerAnalogWithPlatformUser);
+	BIND_PLAYBACK_HANDLER(TEXT("OnControllerButtonWithPlatformUser"), PlayOnControllerButtonPressedWithPlatformUser);
+	BIND_PLAYBACK_HANDLER(TEXT("OnControllerButtonReleasedWithPlatformUser"), PlayOnControllerButtonReleasedWithPlatformUser);
 }
 
 #undef BIND_PLAYBACK_HANDLER
@@ -639,6 +645,7 @@ void FRecordingMessageHandler::PlayOnMotionDetected(FArchive& Ar)
 		Msg.Param5);
 }
 
+PRAGMA_DISABLE_DEPRECATION_WARNINGS
 bool FRecordingMessageHandler::OnControllerAnalog(FGamepadKeyNames::Type KeyName, int32 ControllerId, float AnalogValue)
 {
 	if (IsRecording())
@@ -683,7 +690,6 @@ void FRecordingMessageHandler::PlayOnControllerButtonPressed(FArchive& Ar)
 	OnControllerButtonPressed(FName(*Msg.Param1), Msg.Param2, Msg.Param3);
 }
 
-
 bool FRecordingMessageHandler::OnControllerButtonReleased(FGamepadKeyNames::Type KeyName, int32 ControllerId, bool IsRepeat)
 {
 	if (IsRecording())
@@ -705,6 +711,73 @@ void FRecordingMessageHandler::PlayOnControllerButtonReleased(FArchive& Ar)
 	ThreeParamMsg<FString, int32, bool > Msg(Ar);
 	OnControllerButtonReleased(FName(*Msg.Param1), Msg.Param2, Msg.Param3);
 }
+PRAGMA_ENABLE_DEPRECATION_WARNINGS
+
+bool FRecordingMessageHandler::OnControllerAnalog(FGamepadKeyNames::Type KeyName, FPlatformUserId PlatformUserId, FInputDeviceId InputDeviceId, float AnalogValue)
+{
+	if (IsRecording())
+	{
+		FourParamMsg<FString, int32, int32, float> Msg(KeyName.ToString(), PlatformUserId.GetInternalId(), InputDeviceId.GetId(), AnalogValue);
+		RecordMessage(TEXT("OnControllerAnalogWithPlatformUser"), Msg.AsData());
+	}
+
+	if (bConsumeInput)
+	{
+		return true;
+	}
+	
+	return FProxyMessageHandler::OnControllerAnalog(KeyName, PlatformUserId, InputDeviceId, AnalogValue);
+}
+
+void FRecordingMessageHandler::PlayOnControllerAnalogWithPlatformUser(FArchive& Ar)
+{
+	FourParamMsg<FString, int32, int32, float> Msg(Ar);
+	OnControllerAnalog(FName(*Msg.Param1), FPlatformUserId::CreateFromInternalId(Msg.Param2), FInputDeviceId::CreateFromInternalId(Msg.Param3), Msg.Param4);
+}
+
+bool FRecordingMessageHandler::OnControllerButtonPressed(FGamepadKeyNames::Type KeyName, FPlatformUserId PlatformUserId, FInputDeviceId InputDeviceId, bool IsRepeat)
+{
+	if (IsRecording())
+	{
+		FourParamMsg<FString, int32, int32, bool> Msg(KeyName.ToString(), PlatformUserId.GetInternalId(), InputDeviceId.GetId(), IsRepeat);
+		RecordMessage(TEXT("OnControllerButtonPressedWithPlatformUser"), Msg.AsData());
+	}
+
+	if (bConsumeInput)
+	{
+		return true;
+	}
+
+	return FProxyMessageHandler::OnControllerButtonPressed(KeyName, PlatformUserId, InputDeviceId, IsRepeat);
+}
+
+void FRecordingMessageHandler::PlayOnControllerButtonPressedWithPlatformUser(FArchive& Ar)
+{
+	FourParamMsg<FString, int32, int32, bool> Msg(Ar);
+	OnControllerButtonPressed(FName(*Msg.Param1), FPlatformUserId::CreateFromInternalId(Msg.Param2), FInputDeviceId::CreateFromInternalId(Msg.Param3), Msg.Param4);
+}
+
+bool FRecordingMessageHandler::OnControllerButtonReleased(FGamepadKeyNames::Type KeyName, FPlatformUserId PlatformUserId, FInputDeviceId InputDeviceId, bool IsRepeat)
+{
+	if (IsRecording())
+	{
+		FourParamMsg<FString, int32, int32, bool> Msg(KeyName.ToString(), PlatformUserId.GetInternalId(), InputDeviceId.GetId(), IsRepeat);
+		RecordMessage(TEXT("OnControllerButtonReleasedWithPlatformUser"), Msg.AsData());
+	}
+
+	if (bConsumeInput)
+	{
+		return true;
+	}
+
+	return FProxyMessageHandler::OnControllerButtonReleased(KeyName, PlatformUserId, InputDeviceId, IsRepeat);
+}
+
+void FRecordingMessageHandler::PlayOnControllerButtonReleasedWithPlatformUser(FArchive& Ar)
+{
+	FourParamMsg<FString, int32, int32, bool> Msg(Ar);
+	OnControllerButtonReleased(FName(*Msg.Param1), FPlatformUserId::CreateFromInternalId(Msg.Param2), FInputDeviceId::CreateFromInternalId(Msg.Param3), Msg.Param4);
+}
 
 FWidgetPath FRecordingMessageHandler::FindRoutingMessageWidget(const FVector2D& Location) const
 {
diff --git a/Engine/Plugins/Experimental/RemoteSession/Source/RemoteSession/Private/MessageHandler/RecordingMessageHandler.h b/Engine/Plugins/Experimental/RemoteSession/Source/RemoteSession/Private/MessageHandler/RecordingMessageHandler.h
--- a/Engine/Plugins/Experimental/RemoteSession/Source/RemoteSession/Private/MessageHandler/RecordingMessageHandler.h
+++ b/Engine/Plugins/Experimental/RemoteSession/Source/RemoteSession/Private/MessageHandler/RecordingMessageHandler.h
@@ -125,10 +125,18 @@ public:
 	/**
 	 * Controller handling
 	 */
+	UE_DEPRECATED(5.1, "This version of OnControllerAnalog has been deprecated, please use the one that takes an FPlatformUser and FInputDeviceId instead.")
 	virtual bool OnControllerAnalog(FGamepadKeyNames::Type KeyName, int32 ControllerId, float AnalogValue) override;
+	UE_DEPRECATED(5.1, "This version of OnControllerButtonPressed has been deprecated, please use the one that takes an FPlatformUser and FInputDeviceId instead.")
 	virtual bool OnControllerButtonPressed(FGamepadKeyNames::Type KeyName, int32 ControllerId, bool IsRepeat) override;
+	UE_DEPRECATED(5.1, "This version of OnControllerButtonReleased has been deprecated, please use the one that takes an FPlatformUser and FInputDeviceId instead.")
 	virtual bool OnControllerButtonReleased(FGamepadKeyNames::Type KeyName, int32 ControllerId, bool IsRepeat) override;
 
+	virtual bool OnControllerAnalog(FGamepadKeyNames::Type KeyName, FPlatformUserId PlatformUserId, FInputDeviceId InputDeviceId, float AnalogValue) override;
+	virtual bool OnControllerButtonPressed(FGamepadKeyNames::Type KeyName, FPlatformUserId PlatformUserId, FInputDeviceId InputDeviceId, bool IsRepeat) override;
+	virtual bool OnControllerButtonReleased(FGamepadKeyNames::Type KeyName, FPlatformUserId PlatformUserId, FInputDeviceId InputDeviceId, bool IsRepeat) override;
+	
+	
 protected:
 
 	bool ConvertToNormalizedScreenLocation(const FVector2D& InLocation, FVector2f& OutLocation);
@@ -154,6 +162,10 @@ protected:
 	virtual void PlayOnControllerAnalog(FArchive& Ar);
 	virtual void PlayOnControllerButtonPressed(FArchive& Ar);
 	virtual void PlayOnControllerButtonReleased(FArchive& Ar);
+	
+	virtual void PlayOnControllerAnalogWithPlatformUser(FArchive& Ar);
+	virtual void PlayOnControllerButtonPressedWithPlatformUser(FArchive& Ar);
+	virtual void PlayOnControllerButtonReleasedWithPlatformUser(FArchive& Ar);
 
 	FWidgetPath FindRoutingMessageWidget(const FVector2D& Location) const;
 
