RemoteSessionEditor plugin: Module does not need the ../RemoteSession folder added to private include paths

#rnx

diff --git a/Engine/Plugins/Experimental/RemoteSession/Source/RemoteSessionEditor/RemoteSessionEditor.Build.cs b/Engine/Plugins/Experimental/RemoteSession/Source/RemoteSessionEditor/RemoteSessionEditor.Build.cs
--- a/Engine/Plugins/Experimental/RemoteSession/Source/RemoteSessionEditor/RemoteSessionEditor.Build.cs
+++ b/Engine/Plugins/Experimental/RemoteSession/Source/RemoteSessionEditor/RemoteSessionEditor.Build.cs
@@ -8,13 +8,6 @@ public class RemoteSessionEditor : ModuleRules
 	{
 		PCHUsage = ModuleRules.PCHUsageMode.UseExplicitOrSharedPCHs;
 
-		PrivateIncludePaths.AddRange(
-			new string[] {
-				System.IO.Path.Combine(ModuleDirectory, "../RemoteSession/")
-				// ... add public include paths required here ...
-			}
-		);
-
 		PublicDependencyModuleNames.AddRange(
 			new string[]
 			{
@@ -31,7 +24,7 @@ public class RemoteSessionEditor : ModuleRules
 			{
 				"AssetRegistry",
 				"CoreUObject",
-				"EditorFramework",			
+				"EditorFramework",
 				"Engine",
 				"LevelEditor",
 				"MediaIOCore",
