Fixed non unity/pch errors reported by build farm building non pch and non unity

diff --git a/Engine/Plugins/Experimental/PCG/Source/PCG/Private/Graph/PCGGraphExecutor.cpp b/Engine/Plugins/Experimental/PCG/Source/PCG/Private/Graph/PCGGraphExecutor.cpp
--- a/Engine/Plugins/Experimental/PCG/Source/PCG/Private/Graph/PCGGraphExecutor.cpp
+++ b/Engine/Plugins/Experimental/PCG/Source/PCG/Private/Graph/PCGGraphExecutor.cpp
@@ -19,6 +19,8 @@
 
 #if WITH_EDITOR
 #include "Editor.h"
+#else
+#include "GameFramework/Actor.h"
 #endif
 
 // World partition support for in-editor workflows needs these includes
diff --git a/Engine/Plugins/Experimental/RemoteSession/Source/RemoteSession/Private/RemoteSessionModule.cpp b/Engine/Plugins/Experimental/RemoteSession/Source/RemoteSession/Private/RemoteSessionModule.cpp
--- a/Engine/Plugins/Experimental/RemoteSession/Source/RemoteSession/Private/RemoteSessionModule.cpp
+++ b/Engine/Plugins/Experimental/RemoteSession/Source/RemoteSession/Private/RemoteSessionModule.cpp
@@ -3,20 +3,18 @@
 #include "RemoteSessionModule.h"
 #include "Channels/RemoteSessionChannel.h"
 #include "RemoteSessionHost.h"
+#include "Misc/CommandLine.h"
 #include "Misc/CoreDelegates.h"
+#include "Modules/ModuleManager.h"
 #include "RemoteSessionClient.h"
 
-
-#include "Misc/CommandLine.h"
-
 #if WITH_EDITOR
 	#include "Editor.h"
 	#include "ISettingsModule.h"
 	#include "ISettingsSection.h"
+#else
+	#include "HAL/IConsoleManager.h"
 #endif
-#include "Modules/ModuleManager.h"
-
-#include "Modules/ModuleManager.h"
 
 
 #define LOCTEXT_NAMESPACE "FRemoteSessionModule"
diff --git a/Engine/Plugins/Online/OnlineFramework/Source/Hotfix/Public/UpdateManager.h b/Engine/Plugins/Online/OnlineFramework/Source/Hotfix/Public/UpdateManager.h
--- a/Engine/Plugins/Online/OnlineFramework/Source/Hotfix/Public/UpdateManager.h
+++ b/Engine/Plugins/Online/OnlineFramework/Source/Hotfix/Public/UpdateManager.h
@@ -11,6 +11,7 @@ enum class EPatchCheckResult : uint8;
 struct FOnlineError;
 
 class Error;
+class FUniqueNetId;
 class UGameInstance;
 enum class EHotfixResult : uint8;
 
diff --git a/Engine/Plugins/Runtime/AndroidPermission/Source/AndroidPermission/Private/AndroidPermissionCallbackProxy.cpp b/Engine/Plugins/Runtime/AndroidPermission/Source/AndroidPermission/Private/AndroidPermissionCallbackProxy.cpp
--- a/Engine/Plugins/Runtime/AndroidPermission/Source/AndroidPermission/Private/AndroidPermissionCallbackProxy.cpp
+++ b/Engine/Plugins/Runtime/AndroidPermission/Source/AndroidPermission/Private/AndroidPermissionCallbackProxy.cpp
@@ -9,6 +9,7 @@
 
 #if PLATFORM_ANDROID
 #include "Android/AndroidApplication.h"
+#include "Async/TaskGraphInterfaces.h"
 #endif
 
 static UAndroidPermissionCallbackProxy *pProxy = NULL;
diff --git a/Engine/Plugins/VirtualProduction/DMX/DMXPixelMapping/Source/DMXPixelMappingRuntime/Private/ColorSpace/DMXPixelMappingColorSpace.cpp b/Engine/Plugins/VirtualProduction/DMX/DMXPixelMapping/Source/DMXPixelMappingRuntime/Private/ColorSpace/DMXPixelMappingColorSpace.cpp
--- a/Engine/Plugins/VirtualProduction/DMX/DMXPixelMapping/Source/DMXPixelMappingRuntime/Private/ColorSpace/DMXPixelMappingColorSpace.cpp
+++ b/Engine/Plugins/VirtualProduction/DMX/DMXPixelMapping/Source/DMXPixelMappingRuntime/Private/ColorSpace/DMXPixelMappingColorSpace.cpp
@@ -1,7 +1,5 @@
 // Copyright Epic Games, Inc. All Rights Reserved.
 
-#pragma once
-
 #include "ColorSpace/DMXPixelMappingColorSpace.h"
 
 #include "ColorSpace.h"
diff --git a/Engine/Plugins/VirtualProduction/RemoteControlWebInterface/Source/RemoteControlWebInterface/Private/RCWebInterface.cpp b/Engine/Plugins/VirtualProduction/RemoteControlWebInterface/Source/RemoteControlWebInterface/Private/RCWebInterface.cpp
--- a/Engine/Plugins/VirtualProduction/RemoteControlWebInterface/Source/RemoteControlWebInterface/Private/RCWebInterface.cpp
+++ b/Engine/Plugins/VirtualProduction/RemoteControlWebInterface/Source/RemoteControlWebInterface/Private/RCWebInterface.cpp
@@ -13,6 +13,8 @@
 #if WITH_EDITOR
 #include "Editor.h"
 #include "RCWebInterfaceCustomizations.h"
+#else
+#include "HAL/IConsoleManager.h"
 #endif
 
 #define LOCTEXT_NAMESPACE "FRemoteControlWebInterfaceModule"
