Removing commented out defines.

diff --git a/Engine/Plugins/Experimental/RemoteSession/Source/RemoteSession/Private/Channels/RemoteSessionARCameraChannel.cpp b/Engine/Plugins/Experimental/RemoteSession/Source/RemoteSession/Private/Channels/RemoteSessionARCameraChannel.cpp
--- a/Engine/Plugins/Experimental/RemoteSession/Source/RemoteSession/Private/Channels/RemoteSessionARCameraChannel.cpp
+++ b/Engine/Plugins/Experimental/RemoteSession/Source/RemoteSession/Private/Channels/RemoteSessionARCameraChannel.cpp
@@ -30,9 +30,6 @@
 
 #include "PostProcess/DrawRectangle.h"
 
-//#include "PostProcess/PostProcessMaterial.h"
-//#include "PostProcess/SceneFilterRendering.h"
-
 #define CAMERA_MESSAGE_ADDRESS TEXT("/ARCamera")
 
 TAutoConsoleVariable<int32> CVarJPEGQuality(
