Fixing RemoteSessionARCamera shaders by giving them shader parameters to bind, even though this code path never actually set those parameters to any values.

diff --git a/Engine/Plugins/Experimental/RemoteSession/Source/RemoteSession/Private/Channels/RemoteSessionARCameraChannel.cpp b/Engine/Plugins/Experimental/RemoteSession/Source/RemoteSession/Private/Channels/RemoteSessionARCameraChannel.cpp
--- a/Engine/Plugins/Experimental/RemoteSession/Source/RemoteSession/Private/Channels/RemoteSessionARCameraChannel.cpp
+++ b/Engine/Plugins/Experimental/RemoteSession/Source/RemoteSession/Private/Channels/RemoteSessionARCameraChannel.cpp
@@ -29,6 +29,7 @@
 #include "SceneRenderTargetParameters.h"
 
 #include "PostProcess/DrawRectangle.h"
+#include "PostProcess/PostProcessMaterialInputs.h"
 
 #define CAMERA_MESSAGE_ADDRESS TEXT("/ARCamera")
 
@@ -50,14 +51,19 @@ TAutoConsoleVariable<int32> CVarJPEGGpu(
 	TEXT("1 (default) compresses on the GPU, 0 on the CPU"),
 	ECVF_Default);
 
+BEGIN_SHADER_PARAMETER_STRUCT(FRemoteSessionARCameraMaterialParameters, )
+	SHADER_PARAMETER_STRUCT(FScreenPassTextureViewportParameters, PostProcessOutput)
+	SHADER_PARAMETER_STRUCT_ARRAY(FScreenPassTextureInput, PostProcessInput, [kPostProcessMaterialInputCountMax])
+	SHADER_PARAMETER_STRUCT_ARRAY(FScreenPassTextureInput, PathTracingPostProcessInput, [kPathTracingPostProcessMaterialInputCountMax])
+	SHADER_PARAMETER_SAMPLER(SamplerState, PostProcessInput_BilinearSampler)
+	RENDER_TARGET_BINDING_SLOTS()
+END_SHADER_PARAMETER_STRUCT()
+
 class FPostProcessMaterialShader : public FMaterialShader
 {
 public:
-	FPostProcessMaterialShader() = default;
-	FPostProcessMaterialShader(const ShaderMetaType::CompiledShaderInitializerType& Initializer)
-		: FMaterialShader(Initializer)
-	{
-	}
+	using FParameters = FRemoteSessionARCameraMaterialParameters;
+	SHADER_USE_PARAMETER_STRUCT_WITH_LEGACY_BASE(FPostProcessMaterialShader, FMaterialShader);
 
 	static bool ShouldCompilePermutation(const FMaterialShaderPermutationParameters& Parameters)
 	{
