Fixing editor crashing due to incorrect parsing of RemoteSession plugin motion events.
Fixing RemoteSessionApp compilation due to missing module.

diff --git a/Engine/Plugins/Experimental/RemoteSession/Source/RemoteSession/Private/MessageHandler/RecordingMessageHandler.cpp b/Engine/Plugins/Experimental/RemoteSession/Source/RemoteSession/Private/MessageHandler/RecordingMessageHandler.cpp
--- a/Engine/Plugins/Experimental/RemoteSession/Source/RemoteSession/Private/MessageHandler/RecordingMessageHandler.cpp
+++ b/Engine/Plugins/Experimental/RemoteSession/Source/RemoteSession/Private/MessageHandler/RecordingMessageHandler.cpp
@@ -572,7 +572,12 @@ bool FRecordingMessageHandler::OnTouchGesture(EGestureEvent GestureType, const F
 {
 	if (IsRecording())
 	{
-		FourParamMsg<uint32, FVector2D, float, bool> Msg((uint32)GestureType, Delta, WheelDelta, bIsDirectionInvertedFromDevice);
+		// Touch Delta is FVector2D, but host is parsing FVector2f, so send float to stay compatible with previous versions.
+		FourParamMsg<uint32, FVector2f, float, bool> Msg(
+			(uint32)GestureType,
+			FVector2f(Delta.X, Delta.Y),
+			WheelDelta,
+			bIsDirectionInvertedFromDevice);
 		OutputWriter->RecordMessage(TEXT("OnTouchGesture"), Msg.AsData());
 	}
 
@@ -611,13 +616,17 @@ void FRecordingMessageHandler::PlayOnEndGesture(FArchive& Ar)
 	OnEndGesture();
 }
 
-
 bool FRecordingMessageHandler::OnMotionDetected(const FVector& Tilt, const FVector& RotationRate, const FVector& Gravity, const FVector& Acceleration, int32 ControllerId)
 {
 	if (IsRecording())
 	{
-		FiveParamMsg<FVector, FVector, FVector, FVector, int32> 
-			Msg(Tilt, RotationRate, Gravity, Acceleration, ControllerId);
+		// LWC FVector3D are now doubles, but either legacy client or legacy host might use floats. So we Serialize as floats then create LWC versions in deserialization.
+		FiveParamMsg<FVector3f, FVector3f, FVector3f, FVector3f, int32> Msg(
+			FVector3f(Tilt.X, Tilt.Y, Tilt.Z),
+			FVector3f(RotationRate.X, RotationRate.Y, RotationRate.Z),
+			FVector3f(Gravity.X, Gravity.Y, Gravity.Z),
+			FVector3f(Acceleration.X, Acceleration.Y, Acceleration.Z),
+			ControllerId);
 		OutputWriter->RecordMessage(TEXT("OnMotionDetected"), Msg.AsData());
 	}
 
@@ -631,8 +640,7 @@ bool FRecordingMessageHandler::OnMotionDetected(const FVector& Tilt, const FVect
 
 void FRecordingMessageHandler::PlayOnMotionDetected(FArchive& Ar)
 {
-	// LWC FVector2D are now doubles, but the OSC clients will sent floats. So we deserialize as floats
-	// then create LWC versions below
+	// LWC FVector3D are now doubles, but the OSC clients will sent floats. So we deserialize as floats then create LWC versions below.
 	FiveParamMsg<FVector3f, FVector3f, FVector3f, FVector3f, int32 > Msg(Ar);
 	OnMotionDetected(
 		FVector(Msg.Param1.X, Msg.Param1.Y, Msg.Param1.Z), 
diff --git a/Engine/Plugins/Experimental/RemoteSession/Source/RemoteSession/RemoteSession.Build.cs b/Engine/Plugins/Experimental/RemoteSession/Source/RemoteSession/RemoteSession.Build.cs
--- a/Engine/Plugins/Experimental/RemoteSession/Source/RemoteSession/RemoteSession.Build.cs
+++ b/Engine/Plugins/Experimental/RemoteSession/Source/RemoteSession/RemoteSession.Build.cs
@@ -6,8 +6,6 @@ public class RemoteSession : ModuleRules
 {
 	public RemoteSession(ReadOnlyTargetRules Target) : base(Target)
 	{
-		DefaultBuildSettings = BuildSettingsVersion.V2;
-
 		PrivateIncludePaths.AddRange(
 			new string[] {
 			}
diff --git a/Engine/Plugins/Experimental/RemoteSession/Source/RemoteSessionEditor/RemoteSessionEditor.Build.cs b/Engine/Plugins/Experimental/RemoteSession/Source/RemoteSessionEditor/RemoteSessionEditor.Build.cs
--- a/Engine/Plugins/Experimental/RemoteSession/Source/RemoteSessionEditor/RemoteSessionEditor.Build.cs
+++ b/Engine/Plugins/Experimental/RemoteSession/Source/RemoteSessionEditor/RemoteSessionEditor.Build.cs
@@ -14,6 +14,7 @@ public class RemoteSessionEditor : ModuleRules
 				"Core",
 				"CoreUObject",
 				"RemoteSession",
+				"FieldNotification",
 				"UMG",
 				"UMGEditor",
 			}
