- Fixing remote session window not being spawnable

diff --git a/Engine/Plugins/Experimental/RemoteSession/Source/RemoteSessionEditor/Private/RemoteSessionEditorModule.cpp b/Engine/Plugins/Experimental/RemoteSession/Source/RemoteSessionEditor/Private/RemoteSessionEditorModule.cpp
--- a/Engine/Plugins/Experimental/RemoteSession/Source/RemoteSessionEditor/Private/RemoteSessionEditorModule.cpp
+++ b/Engine/Plugins/Experimental/RemoteSession/Source/RemoteSessionEditor/Private/RemoteSessionEditorModule.cpp
@@ -5,6 +5,9 @@
 
 #include "Modules/ModuleManager.h"
 #include "Widgets/SRemoteSessionStream.h"
+#include "WorkspaceMenuStructure.h"
+#include "WorkspaceMenuStructureModule.h"
+
 
 LLM_DEFINE_TAG(RemoteSession_RemoteSessionEditor);
 #define LOCTEXT_NAMESPACE "FRemoteSessionEditorModule"
@@ -17,7 +20,17 @@ public:
 		LLM_SCOPE_BYTAG(RemoteSession_RemoteSessionEditor);
 
 		FRemoteSessionEditorStyle::Register();
-		SRemoteSessionStream::RegisterNomadTabSpawner();
+
+		{
+			constexpr bool bSortChildren = true;
+			const IWorkspaceMenuStructure& MenuStructure = WorkspaceMenu::GetMenuStructure();
+			TSharedRef<FWorkspaceItem> MenuGroup = MenuStructure.GetLevelEditorVirtualProductionCategory()->AddGroup(
+				LOCTEXT("WorkspaceMenu_VirtualProductionCategory", "VirtualProduction"),
+				FSlateIcon(),
+				bSortChildren);
+
+			SRemoteSessionStream::RegisterNomadTabSpawner(MenuGroup);
+		}
 	}
 
 	virtual void ShutdownModule() override
diff --git a/Engine/Plugins/Experimental/RemoteSession/Source/RemoteSessionEditor/Private/Widgets/SRemoteSessionStream.cpp b/Engine/Plugins/Experimental/RemoteSession/Source/RemoteSessionEditor/Private/Widgets/SRemoteSessionStream.cpp
--- a/Engine/Plugins/Experimental/RemoteSession/Source/RemoteSessionEditor/Private/Widgets/SRemoteSessionStream.cpp
+++ b/Engine/Plugins/Experimental/RemoteSession/Source/RemoteSessionEditor/Private/Widgets/SRemoteSessionStream.cpp
@@ -65,9 +65,9 @@ URemoteSessionStreamWidgetUserData::URemoteSessionStreamWidgetUserData()
 TWeakPtr<SRemoteSessionStream> SRemoteSessionStream::WidgetInstance;
 
 
-void SRemoteSessionStream::RegisterNomadTabSpawner()
+void SRemoteSessionStream::RegisterNomadTabSpawner(TSharedRef<FWorkspaceItem> InWorkspaceItem)
 {
-	auto RegisterTabSpawner = []()
+	auto RegisterTabSpawner = [InWorkspaceItem]()
 	{
 		LLM_SCOPE_BYNAME(TEXT("RemoteSession/RemoteSessionStream"));
 		FLevelEditorModule& LevelEditorModule = FModuleManager::LoadModuleChecked<FLevelEditorModule>(RemoteSessionStream::LevelEditorModuleName);
@@ -76,7 +76,7 @@ void SRemoteSessionStream::RegisterNomadTabSpawner()
 		LevelEditorTabManager->RegisterTabSpawner(RemoteSessionStream::RemoteSessionStreamApp, FOnSpawnTab::CreateStatic(&RemoteSessionStream::CreateMediaFrameworkCaptureCameraViewportTab))
 			.SetDisplayName(LOCTEXT("TabTitle", "Remote Session Stream"))
 			.SetTooltipText(LOCTEXT("TooltipText", "Stream a particular UMG to a Remote Session app."))
-			.SetGroup(WorkspaceMenu::GetMenuStructure().GetToolsCategory())
+			.SetGroup(InWorkspaceItem)
 			.SetIcon(FSlateIcon(FRemoteSessionEditorStyle::GetStyleSetName(), "RemoteSessionStream"));
 	};
 
diff --git a/Engine/Plugins/Experimental/RemoteSession/Source/RemoteSessionEditor/Private/Widgets/SRemoteSessionStream.h b/Engine/Plugins/Experimental/RemoteSession/Source/RemoteSessionEditor/Private/Widgets/SRemoteSessionStream.h
--- a/Engine/Plugins/Experimental/RemoteSession/Source/RemoteSessionEditor/Private/Widgets/SRemoteSessionStream.h
+++ b/Engine/Plugins/Experimental/RemoteSession/Source/RemoteSessionEditor/Private/Widgets/SRemoteSessionStream.h
@@ -32,6 +32,7 @@ class UTextureRenderTarget2D;
 class UUserWidget;
 class UWidgetBlueprint;
 class UWorld;
+class FWorkspaceItem;
 
 enum class EMapChangeType : uint8;
 
@@ -85,7 +86,7 @@ public:
 class SRemoteSessionStream : public SCompoundWidget, FGCObject, FTickableEditorObject
 {
 public:
-	static void RegisterNomadTabSpawner();
+	static void RegisterNomadTabSpawner(TSharedRef<FWorkspaceItem> InWorkspaceItem);
 	static void UnregisterNomadTabSpawner();
 	static TSharedPtr<SRemoteSessionStream> GetPanelInstance();
 
