Fix RemoteSession RecordingMessageHandler not passing Touch events due to the PlatformUserId API change.

diff --git a/Engine/Plugins/Experimental/RemoteSession/Source/RemoteSession/Private/MessageHandler/RecordingMessageHandler.h b/Engine/Plugins/Experimental/RemoteSession/Source/RemoteSession/Private/MessageHandler/RecordingMessageHandler.h
--- a/Engine/Plugins/Experimental/RemoteSession/Source/RemoteSession/Private/MessageHandler/RecordingMessageHandler.h
+++ b/Engine/Plugins/Experimental/RemoteSession/Source/RemoteSession/Private/MessageHandler/RecordingMessageHandler.h
@@ -97,6 +97,12 @@ public:
 
 public:
 
+	/* FRecordingMessageHandler does not support sending PlatformUserId */
+	virtual bool ShouldUsePlatformUserId() const override
+	{
+		return false;
+	}
+
 	/**
 	 * Keyboard handling
 	 */
