Added RemoteSession Hello protocol to sync PixelStreaming version and the Signalling server port.
The RemoteSession host supports configuring and starting PixelStreaming2 in the Editor.

diff --git a/Engine/Plugins/Experimental/BackChannel/Source/BackChannel/Private/Protocol/BackChannelOSCMessage.cpp b/Engine/Plugins/Experimental/BackChannel/Source/BackChannel/Private/Protocol/BackChannelOSCMessage.cpp
--- a/Engine/Plugins/Experimental/BackChannel/Source/BackChannel/Private/Protocol/BackChannelOSCMessage.cpp
+++ b/Engine/Plugins/Experimental/BackChannel/Source/BackChannel/Private/Protocol/BackChannelOSCMessage.cpp
@@ -70,6 +70,12 @@ int FBackChannelOSCMessage::Read(const TCHAR* InName, FString& Value)
 	}
 	else*/
 	{
+		if (TagIndex == TagString.Len())
+		{
+			UE_LOG(LogBackChannel, Error, TEXT("OSCMessage: no more tags!"));
+			return 1;
+		}
+
 		TCHAR CurrentTag = TagString[TagIndex];
 
 		if (CurrentTag != 's')
diff --git a/Engine/Plugins/Experimental/BackChannel/Source/BackChannel/Public/BackChannel/IBackChannelPacket.h b/Engine/Plugins/Experimental/BackChannel/Source/BackChannel/Public/BackChannel/IBackChannelPacket.h
--- a/Engine/Plugins/Experimental/BackChannel/Source/BackChannel/Public/BackChannel/IBackChannelPacket.h
+++ b/Engine/Plugins/Experimental/BackChannel/Source/BackChannel/Public/BackChannel/IBackChannelPacket.h
@@ -72,4 +72,7 @@ public:
 
 	/* Read a block of data from the message */
 	virtual int Read(const TCHAR* InName, void* OutBlob, int32 MaxBlobSize, int32& OutBlobSize) = 0;
+
+	/* Reset the reading position */
+	virtual void ResetReading() = 0;
 };
diff --git a/Engine/Plugins/Experimental/BackChannel/Source/BackChannel/Public/BackChannel/Protocol/OSC/BackChannelOSCMessage.h b/Engine/Plugins/Experimental/BackChannel/Source/BackChannel/Public/BackChannel/Protocol/OSC/BackChannelOSCMessage.h
--- a/Engine/Plugins/Experimental/BackChannel/Source/BackChannel/Public/BackChannel/Protocol/OSC/BackChannelOSCMessage.h
+++ b/Engine/Plugins/Experimental/BackChannel/Source/BackChannel/Public/BackChannel/Protocol/OSC/BackChannelOSCMessage.h
@@ -246,6 +246,12 @@ public:
 	/* Read a blob of data from our arguments */
 	UE_API virtual int Read(const TCHAR* InName, void* InBlob, int32 MaxBlobSize, int32& OutBlobSize) override;	
 
+	/* Reset the reading position */
+	UE_API virtual void ResetReading() override
+	{
+		ResetRead();
+	}
+
 	/* 
 		Read data from the message into a TArray. It must have been serialized by the Read form for TArray(!). 
 		Note - data will be appended to the array.
diff --git a/Engine/Plugins/Experimental/RemoteSession/RemoteSession.uplugin b/Engine/Plugins/Experimental/RemoteSession/RemoteSession.uplugin
--- a/Engine/Plugins/Experimental/RemoteSession/RemoteSession.uplugin
+++ b/Engine/Plugins/Experimental/RemoteSession/RemoteSession.uplugin
@@ -49,6 +49,10 @@
 		{
 			"Name": "XRBase",
 			"Enabled": true
+		},
+		{
+			"Name": "PixelStreaming2",
+			"Enabled": true
 		}
 	]
 }
diff --git a/Engine/Plugins/Experimental/RemoteSession/Source/RemoteSession/Private/RemoteSessionClient.cpp b/Engine/Plugins/Experimental/RemoteSession/Source/RemoteSession/Private/RemoteSessionClient.cpp
--- a/Engine/Plugins/Experimental/RemoteSession/Source/RemoteSession/Private/RemoteSessionClient.cpp
+++ b/Engine/Plugins/Experimental/RemoteSession/Source/RemoteSession/Private/RemoteSessionClient.cpp
@@ -10,7 +10,7 @@
 
 DECLARE_CYCLE_STAT(TEXT("RSClientTick"), STAT_RDClientTick, STATGROUP_Game);
 
-FRemoteSessionClient::FRemoteSessionClient(const TCHAR* InHostAddress)
+FRemoteSessionClient::FRemoteSessionClient(const TCHAR* InHostAddress, bool bUsingPixelStreaming)
 {
 	HostAddress = InHostAddress;
 	ConnectionAttemptTimer = FLT_MAX;		// attempt a connection asap
@@ -24,6 +24,8 @@ FRemoteSessionClient::FRemoteSessionClient(const TCHAR* InHostAddress)
 		HostAddress += FString::Printf(TEXT(":%d"), (int32)IRemoteSessionModule::kDefaultPort);
 	}
 
+	bIsUsingPixelStreaming = bUsingPixelStreaming;
+
 	UE_LOG(LogRemoteSession, Display, TEXT("Will attempt to connect to %s.."), *HostAddress);
 }
 
@@ -142,6 +144,9 @@ void FRemoteSessionClient::BindEndpoints(TBackChannelSharedPtr<IBackChannelConne
 	auto Delegate = FBackChannelRouteDelegate::FDelegate::CreateRaw(this, &FRemoteSessionClient::OnReceiveChannelList);
 	InConnection->AddRouteDelegate(kLegacyChannelSelectionEndPoint, Delegate);
 	InConnection->AddRouteDelegate(kChannelListEndPoint, Delegate);
+
+	Delegate = FBackChannelRouteDelegate::FDelegate::CreateRaw(this, &FRemoteSessionClient::OnReceiveHello);
+	InConnection->AddRouteDelegate(kHelloEndPoint, Delegate);
 }
 
 bool FRemoteSessionClient::ProcessStateChange(const ConnectionState NewState, const ConnectionState OldState)
@@ -161,6 +166,38 @@ bool FRemoteSessionClient::ProcessStateChange(const ConnectionState NewState, co
 	return true;
 }
 
+void FRemoteSessionClient::AppendExtraInfoToBackChannelPacket(TBackChannelSharedPtr<IBackChannelPacket> Packet)
+{
+	if (Packet->GetPath() == kHelloEndPoint)
+	{
+		if (bIsUsingPixelStreaming)
+		{
+			if (FModuleManager::Get().IsModuleLoaded("PixelStreaming2"))
+			{
+				Packet->Write(TEXT("PSVersion"), TEXT("PS2"));
+			}
+			else if (FModuleManager::Get().IsModuleLoaded("PixelStreaming"))
+			{
+				Packet->Write(TEXT("PSVersion"), TEXT("PS1"));
+			}
+		}
+	}
+}
+
+void FRemoteSessionClient::OnReceiveHello(IBackChannelPacket& Message)
+{
+	// Need to read the message values one by one to get the correct data we need
+	Message.ResetReading();
+
+	FString VerString;
+	Message.Read(TEXT("Version"), VerString);
+
+	if (bIsUsingPixelStreaming)
+	{
+		RemotePixelStreamingVersion.Empty();
+		Message.Read(TEXT("PSVersion"), RemotePixelStreamingVersion);
+	}
+}
 
 void FRemoteSessionClient::OnReceiveChannelList(IBackChannelPacket& Message)
 {
diff --git a/Engine/Plugins/Experimental/RemoteSession/Source/RemoteSession/Private/RemoteSessionClient.h b/Engine/Plugins/Experimental/RemoteSession/Source/RemoteSession/Private/RemoteSessionClient.h
--- a/Engine/Plugins/Experimental/RemoteSession/Source/RemoteSession/Private/RemoteSessionClient.h
+++ b/Engine/Plugins/Experimental/RemoteSession/Source/RemoteSession/Private/RemoteSessionClient.h
@@ -13,7 +13,7 @@ class FRemoteSessionClient : public FRemoteSessionRole
 {
 public:
 
-	FRemoteSessionClient(const TCHAR* InHostAddress);
+	FRemoteSessionClient(const TCHAR* InHostAddress, bool bUsingPixelStreaming);
 	~FRemoteSessionClient();
 
 	virtual void Tick(float DeltaTime) override;
@@ -32,6 +32,9 @@ protected:
 
 	void				RequestChannel(const FRemoteSessionChannelInfo& Info);
 
+	void				OnReceiveHello(IBackChannelPacket& Message);
+	virtual void		AppendExtraInfoToBackChannelPacket(TBackChannelSharedPtr<IBackChannelPacket> Packet) override;
+
 	FString				HostAddress;
 	
 	bool				IsConnecting;
diff --git a/Engine/Plugins/Experimental/RemoteSession/Source/RemoteSession/Private/RemoteSessionHost.cpp b/Engine/Plugins/Experimental/RemoteSession/Source/RemoteSession/Private/RemoteSessionHost.cpp
--- a/Engine/Plugins/Experimental/RemoteSession/Source/RemoteSession/Private/RemoteSessionHost.cpp
+++ b/Engine/Plugins/Experimental/RemoteSession/Source/RemoteSession/Private/RemoteSessionHost.cpp
@@ -8,9 +8,26 @@
 #include "BackChannel/Transport/IBackChannelSocketConnection.h"
 #include "RemoteSessionModule.h"
 #include "Trace/Trace.inl"
+#include "Engine/Engine.h"
+
+#define ENABLE_PIXEL_STREAMING_2 1
+
+#if ENABLE_PIXEL_STREAMING_2
+#if WITH_EDITOR
+#include "IPixelStreaming2Module.h"
+#include "IPixelStreaming2EditorModule.h"
+#endif //WITH_EDITOR
+
+#else //ENABLE_PIXEL_STREAMING_2
 
 #if WITH_EDITOR
-#endif
+#include "IPixelStreamingModule.h"
+#include "IPixelStreamingEditorModule.h"
+#endif //WITH_EDITOR
+
+#endif //ENABLE_PIXEL_STREAMING_2
+
+
 
 namespace RemoteSessionEd
 {
@@ -94,6 +111,36 @@ bool FRemoteSessionHost::ProcessStateChange(const ConnectionState NewState, cons
 		IsListenerConnected = true;
 
 		SendChannelListToConnection();
+
+		if (bIsUsingPixelStreaming && !RemotePixelStreamingVersion.IsEmpty())
+		{
+#if ENABLE_PIXEL_STREAMING_2
+			if (FModuleManager::Get().IsModuleLoaded("PixelStreaming2"))
+			{
+				GEngine->Exec(GEngine->GetWorld(), TEXT("PixelStreaming2.WebRTC.DisableReceiveVideo 1"));
+				GEngine->Exec(GEngine->GetWorld(), TEXT("PixelStreaming2.DecoupleFrameRate 1"));
+				GEngine->Exec(GEngine->GetWorld(), TEXT("PixelStreaming2.WebRTC.MinBitrate 5000000"));
+#if WITH_EDITOR
+				if (FModuleManager::Get().IsModuleLoaded("PixelStreaming2Editor"))
+				{
+					bool IsStreaming = false;
+					IPixelStreaming2Module::Get().ForEachStreamer([&](TSharedPtr<IPixelStreaming2Streamer> Streamer) {
+						TSharedPtr<IPixelStreaming2VideoProducer> VideoProducer = Streamer->GetVideoProducer().Pin();
+						if (VideoProducer.IsValid() && (VideoProducer->ToString() == TEXT("the PIE Viewport")))
+						{
+							if (!Streamer->IsStreaming())
+							{
+								Streamer->StartStreaming();
+							}
+						}
+					});
+
+					IPixelStreaming2EditorModule::Get().StartSignalling();
+				}
+#endif //WITH_EDITOR
+			}
+#endif //ENABLE_PIXEL_STREAMING_2
+		}
 	}
 
 	return true;
@@ -102,7 +149,10 @@ bool FRemoteSessionHost::ProcessStateChange(const ConnectionState NewState, cons
 
 void FRemoteSessionHost::BindEndpoints(TBackChannelSharedPtr<IBackChannelConnection> InConnection)
 {
-	FRemoteSessionRole::BindEndpoints(InConnection);	
+	FRemoteSessionRole::BindEndpoints(InConnection);
+
+	auto Delegate = FBackChannelRouteDelegate::FDelegate::CreateRaw(this, &FRemoteSessionHost::OnReceiveHello);
+	InConnection->AddRouteDelegate(kHelloEndPoint, Delegate);
 }
 
 
@@ -154,6 +204,66 @@ void FRemoteSessionHost::SendChannelListToConnection()
 	}
 }
 
+void FRemoteSessionHost::AppendExtraInfoToBackChannelPacket(TBackChannelSharedPtr<IBackChannelPacket> Packet)
+{
+	if (Packet->GetPath() == kHelloEndPoint)
+	{
+#if ENABLE_PIXEL_STREAMING_2
+		if (FModuleManager::Get().IsModuleLoaded("PixelStreaming2"))
+		{
+			Packet->Write(TEXT("PSVersion"), TEXT("PS2"));
+#if WITH_EDITOR
+			if (FModuleManager::Get().IsModuleLoaded("PixelStreaming2Editor"))
+			{
+				FString ViewerPort = FString::FromInt(IPixelStreaming2EditorModule::Get().GetViewerPort());
+				if (ViewerPort.IsEmpty() || !ViewerPort.IsNumeric())
+				{
+					ViewerPort = TEXT("80");
+				}
+				Packet->Write(TEXT("PSHostPort"), ViewerPort);
+			}
+#endif //WITH_EDITOR
+		}
+#else //ENABLE_PIXEL_STREAMING_2
+		if (FModuleManager::Get().IsModuleLoaded("PixelStreaming"))
+		{
+			Packet->Write(TEXT("PSVersion"), TEXT("PS1"));
+#if WITH_EDITOR
+			if (FModuleManager::Get().IsModuleLoaded("PixelStreamingEditor"))
+			{
+				FString ViewerPort;
+				ViewerPort = FString::FromInt(IPixelStreamingEditorModule::Get().GetViewerPort());
+				if (ViewerPort.IsEmpty() || !ViewerPort.IsNumeric())
+				{
+					ViewerPort = TEXT("80");
+				}
+				Packet->Write(TEXT("PSHostPort"), ViewerPort);
+			}
+#endif //WITH_EDITOR
+		}
+#endif //ENABLE_PIXEL_STREAMING_2
+	}
+}
+
+void FRemoteSessionHost::OnReceiveHello(IBackChannelPacket& Message)
+{
+	// Need to read the message values one by one to get the correct data we need
+	Message.ResetReading();
+
+	FString VerString;
+	Message.Read(TEXT("Version"), VerString);
+
+	RemotePixelStreamingVersion.Empty();
+	Message.Read(TEXT("PSVersion"), RemotePixelStreamingVersion);
+	if (RemotePixelStreamingVersion.IsEmpty())
+	{
+		bIsUsingPixelStreaming = false;
+	}
+	else
+	{
+		bIsUsingPixelStreaming = true;
+	}
+}
 
 void FRemoteSessionHost::Tick(float DeltaTime)
 {
diff --git a/Engine/Plugins/Experimental/RemoteSession/Source/RemoteSession/Private/RemoteSessionHost.h b/Engine/Plugins/Experimental/RemoteSession/Source/RemoteSession/Private/RemoteSessionHost.h
--- a/Engine/Plugins/Experimental/RemoteSession/Source/RemoteSession/Private/RemoteSessionHost.h
+++ b/Engine/Plugins/Experimental/RemoteSession/Source/RemoteSession/Private/RemoteSessionHost.h
@@ -39,6 +39,8 @@ protected:
 	
 	bool			ProcessIncomingConnection(TSharedRef<IBackChannelSocketConnection> NewConnection);
 
+	void			OnReceiveHello(IBackChannelPacket& Message);
+	virtual void	AppendExtraInfoToBackChannelPacket(TBackChannelSharedPtr<IBackChannelPacket> Packet) override;
 
 	TSharedPtr<IBackChannelSocketConnection> Listener;
 
diff --git a/Engine/Plugins/Experimental/RemoteSession/Source/RemoteSession/Private/RemoteSessionModule.cpp b/Engine/Plugins/Experimental/RemoteSession/Source/RemoteSession/Private/RemoteSessionModule.cpp
--- a/Engine/Plugins/Experimental/RemoteSession/Source/RemoteSession/Private/RemoteSessionModule.cpp
+++ b/Engine/Plugins/Experimental/RemoteSession/Source/RemoteSession/Private/RemoteSessionModule.cpp
@@ -197,14 +197,14 @@ void FRemoteSessionModule::OnPIEEnded(bool bSimulating)
 	StopHost(TEXT("PIE Ended"));
 }
 
-TSharedPtr<IRemoteSessionRole> FRemoteSessionModule::CreateClient(const TCHAR* RemoteAddress)
+TSharedPtr<IRemoteSessionRole> FRemoteSessionModule::CreateClient(const TCHAR* RemoteAddress, bool bUsingPixelStreaming)
 {
 	// todo - remove this and allow multiple clients (and hosts) to be created
 	if (Client.IsValid())
 	{
 		StopClient(Client, TEXT("New Client"));
 	}
-	Client = MakeShareable(new FRemoteSessionClient(RemoteAddress));
+	Client = MakeShareable(new FRemoteSessionClient(RemoteAddress, bUsingPixelStreaming));
 	return Client;
 }
 
diff --git a/Engine/Plugins/Experimental/RemoteSession/Source/RemoteSession/Private/RemoteSessionModule.h b/Engine/Plugins/Experimental/RemoteSession/Source/RemoteSession/Private/RemoteSessionModule.h
--- a/Engine/Plugins/Experimental/RemoteSession/Source/RemoteSession/Private/RemoteSessionModule.h
+++ b/Engine/Plugins/Experimental/RemoteSession/Source/RemoteSession/Private/RemoteSessionModule.h
@@ -46,7 +46,7 @@ public:
 	void AddChannelFactory(const FStringView InChannelName, ERemoteSessionChannelMode InHostMode, TWeakPtr<IRemoteSessionChannelFactoryWorker> Worker) override;
 	void RemoveChannelFactory(TWeakPtr<IRemoteSessionChannelFactoryWorker> Worker) override;
 
-	virtual TSharedPtr<IRemoteSessionRole>	CreateClient(const TCHAR* RemoteAddress) override;
+	virtual TSharedPtr<IRemoteSessionRole>	CreateClient(const TCHAR* RemoteAddress, bool bUsingPixelStreaming) override;
 	virtual void StopClient(TSharedPtr<IRemoteSessionRole> InClient, const FString& InReason) override;
 
 	virtual void InitHost(const int16 Port = 0) override;
diff --git a/Engine/Plugins/Experimental/RemoteSession/Source/RemoteSession/Private/RemoteSessionRole.cpp b/Engine/Plugins/Experimental/RemoteSession/Source/RemoteSession/Private/RemoteSessionRole.cpp
--- a/Engine/Plugins/Experimental/RemoteSession/Source/RemoteSession/Private/RemoteSessionRole.cpp
+++ b/Engine/Plugins/Experimental/RemoteSession/Source/RemoteSession/Private/RemoteSessionRole.cpp
@@ -51,7 +51,10 @@ struct FRemoteSessionRoleCancellationToken
 };
 
 FRemoteSessionRole::FRemoteSessionRole()
+	:bIsUsingPixelStreaming(false)
 {
+	ExternalHelloRouteDelegateHandles.Empty();
+	ExternalHelloRoutePendingDelegates.Empty();
 }
 
 FRemoteSessionRole::~FRemoteSessionRole()
@@ -240,11 +243,39 @@ FDelegateHandle FRemoteSessionRole::RegisterChannelListDelegate(FOnRemoteSession
 	return ReceiveChannelListDelegate.Add(InDelegate);
 }
 
+void FRemoteSessionRole::RegisterHelloRouteMessageDelegate(FBackChannelRouteDelegate::FDelegate InDelegate)
+{
+	if (BackChannelConnection.IsValid())
+	{
+		TBackChannelSharedPtr<IBackChannelConnection> CurrentConnection = BackChannelConnection.Pin();
+		FDelegateHandle Handle = CurrentConnection->AddRouteDelegate(kHelloEndPoint, InDelegate);
+		ExternalHelloRouteDelegateHandles.Add(Handle);
+	}
+	else
+	{
+		ExternalHelloRoutePendingDelegates.Add(InDelegate);
+	}
+}
+
 void FRemoteSessionRole::RemoveAllDelegates(FDelegateUserObject UserObject)
 {
 	ConnectionChangeDelegate.RemoveAll(UserObject);
 	ChannelChangeDelegate.RemoveAll(UserObject);
 	ReceiveChannelListDelegate.RemoveAll(UserObject);
+
+	if (!BackChannelConnection.IsValid())
+	{
+		return;
+	}
+
+	TBackChannelSharedPtr<IBackChannelConnection> CurrentConnection = BackChannelConnection.Pin();
+	for (FDelegateHandle& Handle : ExternalHelloRouteDelegateHandles)
+	{
+		CurrentConnection->RemoveRouteDelegate(kHelloEndPoint, Handle);
+	}
+
+	ExternalHelloRouteDelegateHandles.Empty();
+	ExternalHelloRoutePendingDelegates.Empty();
 }
 
 void FRemoteSessionRole::RemoveRouteDelegates()
@@ -262,6 +293,12 @@ void FRemoteSessionRole::RemoveRouteDelegates()
 	InConnection->RemoveRouteDelegate(kLegacyVersionEndPoint, RouteDelegates.LegacyReceive);
 	InConnection->RemoveRouteDelegate(kChangeChannelEndPoint, RouteDelegates.ChangeChannel);
 
+	for (FDelegateHandle& Handle : ExternalHelloRouteDelegateHandles)
+	{
+		InConnection->RemoveRouteDelegate(kHelloEndPoint, Handle);
+	}
+	ExternalHelloRouteDelegateHandles.Empty();
+
 	BackChannelConnection = {};
 }
 
@@ -287,6 +324,13 @@ void FRemoteSessionRole::BindEndpoints(TBackChannelSharedPtr<IBackChannelConnect
 	Delegate = FBackChannelRouteDelegate::FDelegate::CreateRaw(this, &FRemoteSessionRole::OnReceiveChangeChannel);
 	RouteDelegates.ChangeChannel = InConnection->AddRouteDelegate(kChangeChannelEndPoint, Delegate);
 
+	for (FBackChannelRouteDelegate::FDelegate& PendingDelegate : ExternalHelloRoutePendingDelegates)
+	{
+		FDelegateHandle Handle = InConnection->AddRouteDelegate(kHelloEndPoint, PendingDelegate);
+		ExternalHelloRouteDelegateHandles.Add(Handle);
+	}
+	ExternalHelloRoutePendingDelegates.Empty();
+
 	BackChannelConnection = InConnection;
 }
 
@@ -371,6 +415,9 @@ void FRemoteSessionRole::SendHello()
 
 		Packet->SetPath(kHelloEndPoint);
 		Packet->Write(TEXT("Version"), IRemoteSessionModule::GetLocalVersion());
+
+		AppendExtraInfoToBackChannelPacket(Packet);
+
 		OSCConnection->SendPacket(Packet);
 	}
 }
@@ -384,6 +431,8 @@ void FRemoteSessionRole::OnReceiveHello(IBackChannelPacket& Message)
 		FBackChannelOSCMessage::SetLegacyMode(false);
 	}
 
+	Message.ResetReading();
+
 	// Read the version
 	FString VersionString;
 	FString VersionErrorMessage;
@@ -457,18 +506,24 @@ void FRemoteSessionRole::OnReceiveGoodbye(IBackChannelPacket& Message)
 
 void FRemoteSessionRole::SendPing()
 {
-	UE_LOG(LogRemoteSession, Verbose, TEXT("Sending ping to %s to check connection"), *Connection->GetDescription());
-	FBackChannelOSCMessage Msg(kPingEndPoint);
-	OSCConnection->SendPacket(Msg);
-	LastPingTime = FPlatformTime::Seconds();
+	if (OSCConnection.IsValid())
+	{
+		UE_LOG(LogRemoteSession, Verbose, TEXT("Sending ping to %s to check connection"), *Connection->GetDescription());
+		FBackChannelOSCMessage Msg(kPingEndPoint);
+		OSCConnection->SendPacket(Msg);
+		LastPingTime = FPlatformTime::Seconds();
+	}
 }
 
 void FRemoteSessionRole::OnReceivePing(IBackChannelPacket& InMessage)
 {
 	// send pong
-	UE_LOG(LogRemoteSession, Verbose, TEXT("Received Ping, Sending Pong"));
-	FBackChannelOSCMessage OutMessage(kPongEndPoint);
-	OSCConnection->SendPacket(OutMessage);
+	if (OSCConnection.IsValid())
+	{
+		UE_LOG(LogRemoteSession, Verbose, TEXT("Received Ping, Sending Pong"));
+		FBackChannelOSCMessage OutMessage(kPongEndPoint);
+		OSCConnection->SendPacket(OutMessage);
+	}
 }
 
 void FRemoteSessionRole::OnReceivePong(IBackChannelPacket& Message)
diff --git a/Engine/Plugins/Experimental/RemoteSession/Source/RemoteSession/Private/RemoteSessionRole.h b/Engine/Plugins/Experimental/RemoteSession/Source/RemoteSession/Private/RemoteSessionRole.h
--- a/Engine/Plugins/Experimental/RemoteSession/Source/RemoteSession/Private/RemoteSessionRole.h
+++ b/Engine/Plugins/Experimental/RemoteSession/Source/RemoteSession/Private/RemoteSessionRole.h
@@ -77,6 +77,9 @@ public:
 	/* Register for notifications whenever a change in the state of a channel occurs */
 	virtual FDelegateHandle RegisterChannelChangeDelegate(FOnRemoteSessionChannelChange::FDelegate InDelegate) override;
 
+	/* Register for notifications whenever a hello route message arrives */
+	virtual void RegisterHelloRouteMessageDelegate(FBackChannelRouteDelegate::FDelegate InDelegate) override;
+
 	/* Unregister all delegates for the specified object */
 	virtual void RemoveAllDelegates(FDelegateUserObject UserObject) override;
 
@@ -125,7 +128,9 @@ protected:
 	void			OnReceiveChannelChanged(IBackChannelPacket& Message);
 	
 	void	        ClearChannels();
-	
+
+	/* Children of RemoteSessionRole can use this function to extend the information that a package will carry */
+	virtual void	AppendExtraInfoToBackChannelPacket(TBackChannelSharedPtr<IBackChannelPacket> Packet) {};
 
 	/* Queues the next state to be processed on the next tick. It's an error to call this when there is another state pending */
 	void			SetPendingState(const ConnectionState InState);
@@ -159,6 +164,10 @@ protected:
 
 	TArray<FRemoteSessionChannelInfo> SupportedChannels;
 
+	bool					bIsUsingPixelStreaming;
+
+	FString					RemotePixelStreamingVersion;
+
 private:
 	void			RemoveRouteDelegates();
 
@@ -184,6 +193,9 @@ private:
 	};
 	FRouteDelegates			RouteDelegates;
 
+	TArray<FBackChannelRouteDelegate::FDelegate> ExternalHelloRoutePendingDelegates;
+	TArray<FDelegateHandle> ExternalHelloRouteDelegateHandles;
+
 	ConnectionState			CurrentState = ConnectionState::Disconnected;
 
 	ConnectionState			PendingState = ConnectionState::Unknown;
diff --git a/Engine/Plugins/Experimental/RemoteSession/Source/RemoteSession/Public/IRemoteSessionRole.h b/Engine/Plugins/Experimental/RemoteSession/Source/RemoteSession/Public/IRemoteSessionRole.h
--- a/Engine/Plugins/Experimental/RemoteSession/Source/RemoteSession/Public/IRemoteSessionRole.h
+++ b/Engine/Plugins/Experimental/RemoteSession/Source/RemoteSession/Public/IRemoteSessionRole.h
@@ -3,6 +3,7 @@
 #pragma once
 
 #include "RemoteSessionTypes.h"
+#include "BackChannel/Types.h"
 
 class IRemoteSessionChannel;
 
@@ -31,6 +32,8 @@ public:
 	/* Register for notifications whenever a change in the state of a channel occurs */
 	virtual FDelegateHandle RegisterChannelChangeDelegate(FOnRemoteSessionChannelChange::FDelegate InDelegate) = 0;
 
+	virtual void RegisterHelloRouteMessageDelegate(FBackChannelRouteDelegate::FDelegate InDelegate) = 0;
+
 	/* Unregister all delegates for the specified object */
 	virtual void RemoveAllDelegates(FDelegateUserObject UserObject) = 0;
 
diff --git a/Engine/Plugins/Experimental/RemoteSession/Source/RemoteSession/Public/RemoteSession.h b/Engine/Plugins/Experimental/RemoteSession/Source/RemoteSession/Public/RemoteSession.h
--- a/Engine/Plugins/Experimental/RemoteSession/Source/RemoteSession/Public/RemoteSession.h
+++ b/Engine/Plugins/Experimental/RemoteSession/Source/RemoteSession/Public/RemoteSession.h
@@ -46,7 +46,7 @@ public:
 	/** Client implementation */
 	
 	/** Returns a reference to the client role (if any) */
-	virtual TSharedPtr<IRemoteSessionRole> CreateClient(const TCHAR* RemoteAddress) = 0;
+	virtual TSharedPtr<IRemoteSessionRole> CreateClient(const TCHAR* RemoteAddress, bool bUsingPixelStreaming = false) = 0;
 
 	/** Stops the client. After this CreateClient() must be called if a new connection is desired */
 	virtual void StopClient(TSharedPtr<IRemoteSessionRole> Client, const FString& InReason) = 0;
diff --git a/Engine/Plugins/Experimental/RemoteSession/Source/RemoteSession/RemoteSession.Build.cs b/Engine/Plugins/Experimental/RemoteSession/Source/RemoteSession/RemoteSession.Build.cs
--- a/Engine/Plugins/Experimental/RemoteSession/Source/RemoteSession/RemoteSession.Build.cs
+++ b/Engine/Plugins/Experimental/RemoteSession/Source/RemoteSession/RemoteSession.Build.cs
@@ -51,6 +51,15 @@ public class RemoteSession : ModuleRules
 			//reference the module "MyModule"
 			PrivateDependencyModuleNames.Add("EditorFramework");
 			PrivateDependencyModuleNames.Add("UnrealEd");
+
+			//PrivateDependencyModuleNames.Add("PixelStreaming");
+			//PrivateDependencyModuleNames.Add("PixelStreamingServers");
+			//PrivateDependencyModuleNames.Add("PixelStreamingEditor");
+
+			PrivateDependencyModuleNames.Add("PixelStreaming2");
+			PrivateDependencyModuleNames.Add("PixelStreaming2Settings");
+			PrivateDependencyModuleNames.Add("PixelStreaming2Servers");
+			PrivateDependencyModuleNames.Add("PixelStreaming2Editor");
 		}
 	}
 }
