disable remote session on arm64 as pixel streaming isn't available on arm64 yet

diff --git a/Engine/Plugins/Experimental/RemoteSession/RemoteSession.uplugin b/Engine/Plugins/Experimental/RemoteSession/RemoteSession.uplugin
--- a/Engine/Plugins/Experimental/RemoteSession/RemoteSession.uplugin
+++ b/Engine/Plugins/Experimental/RemoteSession/RemoteSession.uplugin
@@ -20,6 +20,7 @@
 			"Name": "RemoteSession",
 			"Type": "Runtime",
 			"LoadingPhase": "PostConfigInit",
+			"PlatformArchitectureDenyList": [ "Win64:arm64" ],
 			"PlatformDenyList": [
 				"TVOS"
 			]
@@ -28,6 +29,7 @@
 			"Name": "RemoteSessionEditor",
 			"Type": "Editor",
 			"LoadingPhase": "Default",
+			"PlatformArchitectureDenyList": [ "Win64:arm64" ],
 			"PlatformDenyList": [
 				"TVOS"
 			]
