RemoteSessionApp: Negotiate PixelStreaming availability via Hello handshake

Removed checkbox. Host and client exchange PS capability in the Hello message and the client auto-enables PS2
if the host supports it, falling back to MJPEG otherwise.

- Force rotation to portrait or landscape during connection to match incoming resolution
- Removed bUsingPixelStreaming parameter from CreateClient/FRemoteSessionClient
- Client always advertises PS2 capability in Hello if module is loaded
- Client reads host PSVersion in OnReceiveHello and negotiates (PS2 only, PS1 dropped)
- OnPixelStreamingEnabled BP event called during connection if PS is enabled.
- Fixed issue where Hello delegate wouldn't get called on a reconnect
- GameMode respects pixel.streaming=0 command line opt-out
- Removed XR/LiveLink tracking code (replaced by LiveLink app)
- Removed AppleARKit, VirtualProductionUtilities, VirtualCamera plugins

#virtualized
#Synced-CL 51807380

diff --git a/Engine/Plugins/Experimental/BackChannel/Source/BackChannel/Private/Protocol/BackChannelOSCMessage.cpp b/Engine/Plugins/Experimental/BackChannel/Source/BackChannel/Private/Protocol/BackChannelOSCMessage.cpp
--- a/Engine/Plugins/Experimental/BackChannel/Source/BackChannel/Private/Protocol/BackChannelOSCMessage.cpp
+++ b/Engine/Plugins/Experimental/BackChannel/Source/BackChannel/Private/Protocol/BackChannelOSCMessage.cpp
@@ -72,7 +72,7 @@ int FBackChannelOSCMessage::Read(const TCHAR* InName, FString& Value)
 	{
 		if (TagIndex == TagString.Len())
 		{
-			UE_LOG(LogBackChannel, Error, TEXT("OSCMessage: no more tags!"));
+			UE_LOG(LogBackChannel, Warning, TEXT("OSCMessage(%s): no more tags to read (expected 's', have %d tags). This can happen when the sender uses an older protocol version."), *Address, TagString.Len());
 			return 1;
 		}
 
diff --git a/Engine/Plugins/Experimental/RemoteSession/Source/RemoteSession/Private/RemoteSessionClient.cpp b/Engine/Plugins/Experimental/RemoteSession/Source/RemoteSession/Private/RemoteSessionClient.cpp
--- a/Engine/Plugins/Experimental/RemoteSession/Source/RemoteSession/Private/RemoteSessionClient.cpp
+++ b/Engine/Plugins/Experimental/RemoteSession/Source/RemoteSession/Private/RemoteSessionClient.cpp
@@ -10,7 +10,7 @@
 
 DECLARE_CYCLE_STAT(TEXT("RSClientTick"), STAT_RDClientTick, STATGROUP_Game);
 
-FRemoteSessionClient::FRemoteSessionClient(const TCHAR* InHostAddress, bool bUsingPixelStreaming)
+FRemoteSessionClient::FRemoteSessionClient(const TCHAR* InHostAddress)
 {
 	HostAddress = InHostAddress;
 	ConnectionAttemptTimer = FLT_MAX;		// attempt a connection asap
@@ -24,7 +24,8 @@ FRemoteSessionClient::FRemoteSessionClient(const TCHAR* InHostAddress, bool bUsi
 		HostAddress += FString::Printf(TEXT(":%d"), (int32)IRemoteSessionModule::kDefaultPort);
 	}
 
-	bIsUsingPixelStreaming = bUsingPixelStreaming;
+	// OnReceiveHello will enable PixelStreaming if supported by both host and client
+	bIsUsingPixelStreaming = false;
 
 	UE_LOG(LogRemoteSession, Display, TEXT("Will attempt to connect to %s.."), *HostAddress);
 }
@@ -170,16 +171,9 @@ void FRemoteSessionClient::AppendExtraInfoToBackChannelPacket(TBackChannelShared
 {
 	if (Packet->GetPath() == kHelloEndPoint)
 	{
-		if (bIsUsingPixelStreaming)
+		if (FModuleManager::Get().IsModuleLoaded("PixelStreaming2"))
 		{
-			if (FModuleManager::Get().IsModuleLoaded("PixelStreaming2"))
-			{
-				Packet->Write(TEXT("PSVersion"), TEXT("PS2"));
-			}
-			else if (FModuleManager::Get().IsModuleLoaded("PixelStreaming"))
-			{
-				Packet->Write(TEXT("PSVersion"), TEXT("PS1"));
-			}
+			Packet->Write(TEXT("PSVersion"), TEXT("PS2"));
 		}
 	}
 }
@@ -192,11 +186,17 @@ void FRemoteSessionClient::OnReceiveHello(IBackChannelPacket& Message)
 	FString VerString;
 	Message.Read(TEXT("Version"), VerString);
 
-	if (bIsUsingPixelStreaming)
-	{
-		RemotePixelStreamingVersion.Empty();
-		Message.Read(TEXT("PSVersion"), RemotePixelStreamingVersion);
-	}
+	RemotePixelStreamingVersion.Empty();
+	Message.Read(TEXT("PSVersion"), RemotePixelStreamingVersion);
+
+	// Negotiate: use PS if host supports PS2 and we have the PS2 module
+	bool bLocalPSCapable = FModuleManager::Get().IsModuleLoaded("PixelStreaming2");
+	bIsUsingPixelStreaming = bLocalPSCapable && RemotePixelStreamingVersion == TEXT("PS2");
+
+	UE_LOG(LogRemoteSession, Log, TEXT("Negotiated streaming mode: %s (host PSVersion: %s, local PS2 module: %s)"),
+		bIsUsingPixelStreaming ? TEXT("PixelStreaming") : TEXT("MJPEG"),
+		RemotePixelStreamingVersion.IsEmpty() ? TEXT("none") : *RemotePixelStreamingVersion,
+		bLocalPSCapable ? TEXT("loaded") : TEXT("not loaded"));
 }
 
 void FRemoteSessionClient::OnReceiveChannelList(IBackChannelPacket& Message)
diff --git a/Engine/Plugins/Experimental/RemoteSession/Source/RemoteSession/Private/RemoteSessionClient.h b/Engine/Plugins/Experimental/RemoteSession/Source/RemoteSession/Private/RemoteSessionClient.h
--- a/Engine/Plugins/Experimental/RemoteSession/Source/RemoteSession/Private/RemoteSessionClient.h
+++ b/Engine/Plugins/Experimental/RemoteSession/Source/RemoteSession/Private/RemoteSessionClient.h
@@ -13,7 +13,7 @@ class FRemoteSessionClient : public FRemoteSessionRole
 {
 public:
 
-	FRemoteSessionClient(const TCHAR* InHostAddress, bool bUsingPixelStreaming);
+	FRemoteSessionClient(const TCHAR* InHostAddress);
 	~FRemoteSessionClient();
 
 	virtual void Tick(float DeltaTime) override;
diff --git a/Engine/Plugins/Experimental/RemoteSession/Source/RemoteSession/Private/RemoteSessionModule.cpp b/Engine/Plugins/Experimental/RemoteSession/Source/RemoteSession/Private/RemoteSessionModule.cpp
--- a/Engine/Plugins/Experimental/RemoteSession/Source/RemoteSession/Private/RemoteSessionModule.cpp
+++ b/Engine/Plugins/Experimental/RemoteSession/Source/RemoteSession/Private/RemoteSessionModule.cpp
@@ -197,14 +197,14 @@ void FRemoteSessionModule::OnPIEEnded(bool bSimulating)
 	StopHost(TEXT("PIE Ended"));
 }
 
-TSharedPtr<IRemoteSessionRole> FRemoteSessionModule::CreateClient(const TCHAR* RemoteAddress, bool bUsingPixelStreaming)
+TSharedPtr<IRemoteSessionRole> FRemoteSessionModule::CreateClient(const TCHAR* RemoteAddress)
 {
 	// todo - remove this and allow multiple clients (and hosts) to be created
 	if (Client.IsValid())
 	{
 		StopClient(Client, TEXT("New Client"));
 	}
-	Client = MakeShareable(new FRemoteSessionClient(RemoteAddress, bUsingPixelStreaming));
+	Client = MakeShareable(new FRemoteSessionClient(RemoteAddress));
 	return Client;
 }
 
diff --git a/Engine/Plugins/Experimental/RemoteSession/Source/RemoteSession/Private/RemoteSessionModule.h b/Engine/Plugins/Experimental/RemoteSession/Source/RemoteSession/Private/RemoteSessionModule.h
--- a/Engine/Plugins/Experimental/RemoteSession/Source/RemoteSession/Private/RemoteSessionModule.h
+++ b/Engine/Plugins/Experimental/RemoteSession/Source/RemoteSession/Private/RemoteSessionModule.h
@@ -23,16 +23,11 @@ private:
 	TSharedPtr<FRemoteSessionHost>							Host;
 	TSharedPtr<FRemoteSessionClient>						Client;
 
-	//TArray<TWeakPtr<IRemoteSessionChannelFactoryWorker>>	FactoryWorkers;
-	//TArray<TSharedPtr<IRemoteSessionChannelFactoryWorker>>	BuiltInFactory;
-
 	int32													DefaultPort = IRemoteSessionModule::kDefaultPort;
 
 	bool													bAutoHostWithPIE = true;
 	bool													bAutoHostWithGame = true;
 
-	//TArray<FRemoteSessionChannelInfo>						ProgramaticallySupportedChannels;
-
 	FDelegateHandle PostPieDelegate;
 	FDelegateHandle EndPieDelegate;
 	FDelegateHandle GameStartDelegate;
@@ -46,7 +41,7 @@ public:
 	void AddChannelFactory(const FStringView InChannelName, ERemoteSessionChannelMode InHostMode, TWeakPtr<IRemoteSessionChannelFactoryWorker> Worker) override;
 	void RemoveChannelFactory(TWeakPtr<IRemoteSessionChannelFactoryWorker> Worker) override;
 
-	virtual TSharedPtr<IRemoteSessionRole>	CreateClient(const TCHAR* RemoteAddress, bool bUsingPixelStreaming) override;
+	virtual TSharedPtr<IRemoteSessionRole>	CreateClient(const TCHAR* RemoteAddress) override;
 	virtual void StopClient(TSharedPtr<IRemoteSessionRole> InClient, const FString& InReason) override;
 
 	virtual void InitHost(const int16 Port = 0) override;
diff --git a/Engine/Plugins/Experimental/RemoteSession/Source/RemoteSession/Private/RemoteSessionRole.cpp b/Engine/Plugins/Experimental/RemoteSession/Source/RemoteSession/Private/RemoteSessionRole.cpp
--- a/Engine/Plugins/Experimental/RemoteSession/Source/RemoteSession/Private/RemoteSessionRole.cpp
+++ b/Engine/Plugins/Experimental/RemoteSession/Source/RemoteSession/Private/RemoteSessionRole.cpp
@@ -54,7 +54,7 @@ FRemoteSessionRole::FRemoteSessionRole()
 	:bIsUsingPixelStreaming(false)
 {
 	ExternalHelloRouteDelegateHandles.Empty();
-	ExternalHelloRoutePendingDelegates.Empty();
+
 }
 
 FRemoteSessionRole::~FRemoteSessionRole()
@@ -245,16 +245,15 @@ FDelegateHandle FRemoteSessionRole::RegisterChannelListDelegate(FOnRemoteSession
 
 void FRemoteSessionRole::RegisterHelloRouteMessageDelegate(FBackChannelRouteDelegate::FDelegate InDelegate)
 {
+	// Keep a persistent copy so the delegate survives reconnects
+	ExternalHelloRouteDelegates.Add(InDelegate);
+
 	if (BackChannelConnection.IsValid())
 	{
 		TBackChannelSharedPtr<IBackChannelConnection> CurrentConnection = BackChannelConnection.Pin();
 		FDelegateHandle Handle = CurrentConnection->AddRouteDelegate(kHelloEndPoint, InDelegate);
 		ExternalHelloRouteDelegateHandles.Add(Handle);
 	}
-	else
-	{
-		ExternalHelloRoutePendingDelegates.Add(InDelegate);
-	}
 }
 
 void FRemoteSessionRole::RemoveAllDelegates(FDelegateUserObject UserObject)
@@ -263,6 +262,12 @@ void FRemoteSessionRole::RemoveAllDelegates(FDelegateUserObject UserObject)
 	ChannelChangeDelegate.RemoveAll(UserObject);
 	ReceiveChannelListDelegate.RemoveAll(UserObject);
 
+	// Purge persistent delegates bound to this object so they aren't re-registered on reconnect
+	ExternalHelloRouteDelegates.RemoveAll([UserObject](const FBackChannelRouteDelegate::FDelegate& Delegate)
+	{
+		return Delegate.IsBoundToObject(UserObject);
+	});
+
 	if (!BackChannelConnection.IsValid())
 	{
 		return;
@@ -275,7 +280,6 @@ void FRemoteSessionRole::RemoveAllDelegates(FDelegateUserObject UserObject)
 	}
 
 	ExternalHelloRouteDelegateHandles.Empty();
-	ExternalHelloRoutePendingDelegates.Empty();
 }
 
 void FRemoteSessionRole::RemoveRouteDelegates()
@@ -324,12 +328,14 @@ void FRemoteSessionRole::BindEndpoints(TBackChannelSharedPtr<IBackChannelConnect
 	Delegate = FBackChannelRouteDelegate::FDelegate::CreateRaw(this, &FRemoteSessionRole::OnReceiveChangeChannel);
 	RouteDelegates.ChangeChannel = InConnection->AddRouteDelegate(kChangeChannelEndPoint, Delegate);
 
-	for (FBackChannelRouteDelegate::FDelegate& PendingDelegate : ExternalHelloRoutePendingDelegates)
+	// Re-register all external Hello delegates (survives reconnects)
+	ExternalHelloRouteDelegateHandles.Empty();
+	for (FBackChannelRouteDelegate::FDelegate& ExternalDelegate : ExternalHelloRouteDelegates)
 	{
-		FDelegateHandle Handle = InConnection->AddRouteDelegate(kHelloEndPoint, PendingDelegate);
+		FDelegateHandle Handle = InConnection->AddRouteDelegate(kHelloEndPoint, ExternalDelegate);
 		ExternalHelloRouteDelegateHandles.Add(Handle);
 	}
-	ExternalHelloRoutePendingDelegates.Empty();
+
 
 	BackChannelConnection = InConnection;
 }
diff --git a/Engine/Plugins/Experimental/RemoteSession/Source/RemoteSession/Private/RemoteSessionRole.h b/Engine/Plugins/Experimental/RemoteSession/Source/RemoteSession/Private/RemoteSessionRole.h
--- a/Engine/Plugins/Experimental/RemoteSession/Source/RemoteSession/Private/RemoteSessionRole.h
+++ b/Engine/Plugins/Experimental/RemoteSession/Source/RemoteSession/Private/RemoteSessionRole.h
@@ -193,7 +193,7 @@ private:
 	};
 	FRouteDelegates			RouteDelegates;
 
-	TArray<FBackChannelRouteDelegate::FDelegate> ExternalHelloRoutePendingDelegates;
+	TArray<FBackChannelRouteDelegate::FDelegate> ExternalHelloRouteDelegates;
 	TArray<FDelegateHandle> ExternalHelloRouteDelegateHandles;
 
 	ConnectionState			CurrentState = ConnectionState::Disconnected;
diff --git a/Engine/Plugins/Experimental/RemoteSession/Source/RemoteSession/Public/RemoteSession.h b/Engine/Plugins/Experimental/RemoteSession/Source/RemoteSession/Public/RemoteSession.h
--- a/Engine/Plugins/Experimental/RemoteSession/Source/RemoteSession/Public/RemoteSession.h
+++ b/Engine/Plugins/Experimental/RemoteSession/Source/RemoteSession/Public/RemoteSession.h
@@ -46,7 +46,7 @@ public:
 	/** Client implementation */
 	
 	/** Returns a reference to the client role (if any) */
-	virtual TSharedPtr<IRemoteSessionRole> CreateClient(const TCHAR* RemoteAddress, bool bUsingPixelStreaming = false) = 0;
+	virtual TSharedPtr<IRemoteSessionRole> CreateClient(const TCHAR* RemoteAddress) = 0;
 
 	/** Stops the client. After this CreateClient() must be called if a new connection is desired */
 	virtual void StopClient(TSharedPtr<IRemoteSessionRole> Client, const FString& InReason) = 0;
