Fix support for Unreal Remote connecting to Mobile Session in shipping builds
#Synced-CL 51949918

#Synced-CL 51949918

diff --git a/Engine/Plugins/Experimental/RemoteSession/Source/RemoteSession/Private/RemoteSessionModule.cpp b/Engine/Plugins/Experimental/RemoteSession/Source/RemoteSession/Private/RemoteSessionModule.cpp
--- a/Engine/Plugins/Experimental/RemoteSession/Source/RemoteSession/Private/RemoteSessionModule.cpp
+++ b/Engine/Plugins/Experimental/RemoteSession/Source/RemoteSession/Private/RemoteSessionModule.cpp
@@ -291,9 +291,7 @@ TSharedPtr<IRemoteSessionUnmanagedRole> FRemoteSessionModule::CreateHost(TArray<
 TSharedPtr<FRemoteSessionHost> FRemoteSessionModule::CreateHostInternal(TArray<FRemoteSessionChannelInfo> SupportedChannels, int32 Port) const
 {
 #if UE_BUILD_SHIPPING
-
 	const URemoteSessionSettings* Settings = URemoteSessionSettings::StaticClass()->GetDefaultObject<URemoteSessionSettings>();
-	bool bAllowInShipping = false;
 	if (Settings->bAllowInShipping == false)
 	{
 		UE_LOG(LogRemoteSession, Log, TEXT("RemoteSession is disabled. Shipping=1"));
