RemoteSessionApp auto-discovers hosts on iOS/Android using NetworkServiceDiscovery
Added UI to support auto-discovery
Also enabled NetworkServiceDiscovery to QAGame for testing

#virtualized
#Synced-CL 52204226

diff --git a/Engine/Plugins/Experimental/RemoteSession/Source/RemoteSession/Private/RemoteSessionModule.cpp b/Engine/Plugins/Experimental/RemoteSession/Source/RemoteSession/Private/RemoteSessionModule.cpp
--- a/Engine/Plugins/Experimental/RemoteSession/Source/RemoteSession/Private/RemoteSessionModule.cpp
+++ b/Engine/Plugins/Experimental/RemoteSession/Source/RemoteSession/Private/RemoteSessionModule.cpp
@@ -7,6 +7,7 @@
 #include "Misc/CoreDelegates.h"
 #include "Modules/ModuleManager.h"
 #include "RemoteSessionClient.h"
+#include "Misc/App.h"
 
 #if WITH_NETWORK_SERVICE_DISCOVERY
 #include "INetworkServiceDiscovery.h"
@@ -283,10 +284,11 @@ void FRemoteSessionModule::InitHost(const int16 Port /*= 0*/)
 void FRemoteSessionModule::StopHost(const FString& InReason)
 { 
 #if WITH_NETWORK_SERVICE_DISCOVERY
-	// Stop mDNS advertisement
+	// Stop mDNS advertisement — use the same service name as registration
 	if (INetworkServiceDiscoveryModule* Discovery = FModuleManager::GetModulePtr<INetworkServiceDiscoveryModule>("NetworkServiceDiscovery"))
 	{
-		Discovery->UnregisterService();
+		FString ServiceName = FString::Printf(TEXT("%s on %s"), FApp::GetProjectName(), FPlatformProcess::ComputerName());
+		Discovery->UnregisterService(ServiceName);
 	}
 #endif
 
